AddContentDialog - A dialog and lower level infrastructure for adding additional content to an existing game.  Currently there is only support for the new "feature pack" upack files.

diff --git a/Engine/Source/Editor/AddContentDialog/AddContentDialog.Build.cs b/Engine/Source/Editor/AddContentDialog/AddContentDialog.Build.cs
new file mode 100644
--- /dev/null
+++ b/Engine/Source/Editor/AddContentDialog/AddContentDialog.Build.cs
@@ -0,0 +1,50 @@
+// Copyright 1998-2014 Epic Games, Inc. All Rights Reserved.
+
+using UnrealBuildTool;
+
+public class AddContentDialog : ModuleRules
+{
+	public AddContentDialog(TargetInfo Target)
+	{
+		PrivateIncludePathModuleNames.AddRange(
+			new string[] {
+				"AssetTools",
+			}
+		);
+
+		PrivateDependencyModuleNames.AddRange(
+			new string[] {
+				"Core",
+				"Slate",
+				"SlateCore",
+				"InputCore",
+				"Json",
+				"EditorStyle",
+				"DirectoryWatcher",
+				"DesktopPlatform",
+				"PakFile",
+				"ImageWrapper",
+				"UnrealEd"
+			}
+		);
+
+		DynamicallyLoadedModuleNames.AddRange(
+			new string[] {
+			}
+		);
+		
+		PublicIncludePathModuleNames.AddRange(
+			new string[] {
+			}
+		);
+
+		PrivateIncludePaths.AddRange(
+			new string[] {
+				"Editor/AddContentDialog/Private",
+				"Editor/AddContentDialog/Private/ViewModels",
+				"Editor/AddContentDialog/Private/ContentSourceProviders/AssetPack",
+				"Editor/AddContentDialog/Private/ContentSourceProviders/FeaturePack",
+			}
+		);
+	}
+}
diff --git a/Engine/Source/Editor/AddContentDialog/Private/AddContentDialogModule.cpp b/Engine/Source/Editor/AddContentDialog/Private/AddContentDialogModule.cpp
new file mode 100644
--- /dev/null
+++ b/Engine/Source/Editor/AddContentDialog/Private/AddContentDialogModule.cpp
@@ -0,0 +1,44 @@
+// Copyright 1998-2014 Epic Games, Inc. All Rights Reserved.
+
+#include "AddContentDialogPCH.h"
+
+#include "AssertionMacros.h"
+#include "FeaturePackContentSourceProvider.h"
+#include "ModuleManager.h"
+#include "SDockTab.h"
+
+#define LOCTEXT_NAMESPACE "AddContentDialog"
+
+class FAddContentDialogModule : public IAddContentDialogModule
+{
+
+public:
+	virtual void StartupModule() override
+	{
+		ContentSourceProviderManager = TSharedPtr<FContentSourceProviderManager>(new FContentSourceProviderManager());
+		ContentSourceProviderManager->RegisterContentSourceProvider(MakeShareable(new FFeaturePackContentSourceProvider()));
+		FAddContentDialogStyle::Initialize();
+	}
+
+	virtual void ShutdownModule() override
+	{
+		FAddContentDialogStyle::Shutdown();
+	}
+
+	virtual TSharedRef<FContentSourceProviderManager> GetContentSourceProviderManager() override
+	{
+		return ContentSourceProviderManager.ToSharedRef();
+	}
+
+	virtual TSharedRef<SWindow> CreateDialogWindow() override
+	{
+		return SNew(SAddContentDialog);
+	}
+
+private:
+	TSharedPtr<FContentSourceProviderManager> ContentSourceProviderManager;
+	FString TransferFilePath = FPaths::Combine(FPlatformProcess::UserSettingsDir(), TEXT("UnrealEngineLauncher"), TEXT("com"), TEXT("transfer.tmp"));
+
+};
+
+IMPLEMENT_MODULE(FAddContentDialogModule, AddContentDialog);
\ No newline at end of file
diff --git a/Engine/Source/Editor/AddContentDialog/Private/AddContentDialogPCH.h b/Engine/Source/Editor/AddContentDialog/Private/AddContentDialogPCH.h
new file mode 100644
--- /dev/null
+++ b/Engine/Source/Editor/AddContentDialog/Private/AddContentDialogPCH.h
@@ -0,0 +1,22 @@
+// Copyright 1998-2014 Epic Games, Inc. All Rights Reserved.
+
+#pragma once
+
+#include "SlateBasics.h"
+#include "EditorStyle.h"
+#include "TextFilter.h"
+
+#include "IAddContentDialogModule.h"
+
+#include "IContentSource.h"
+#include "IContentSourceProvider.h"
+#include "ContentSourceProviderManager.h"
+
+#include "AddContentDialogStyle.h"
+#include "CategoryViewModel.h"
+#include "ContentSourceViewModel.h"
+#include "AddContentWidgetViewModel.h"
+#include "ContentSourceDragDropOp.h"
+#include "SAddContentWidget.h"
+#include "SAddContentDialog.h"
+
diff --git a/Engine/Source/Editor/AddContentDialog/Private/AddContentDialogStyle.cpp b/Engine/Source/Editor/AddContentDialog/Private/AddContentDialogStyle.cpp
new file mode 100644
--- /dev/null
+++ b/Engine/Source/Editor/AddContentDialog/Private/AddContentDialogStyle.cpp
@@ -0,0 +1,104 @@
+// Copyright 1998-2014 Epic Games, Inc. All Rights Reserved.
+
+#include "AddContentDialogPCH.h"
+
+#include "SlateGameResources.h"
+
+TSharedPtr< FSlateStyleSet > FAddContentDialogStyle::AddContentDialogStyleInstance = NULL;
+
+void FAddContentDialogStyle::Initialize()
+{
+	if ( !AddContentDialogStyleInstance.IsValid() )
+	{
+		AddContentDialogStyleInstance = Create();
+		FSlateStyleRegistry::RegisterSlateStyle( *AddContentDialogStyleInstance );
+	}
+}
+
+void FAddContentDialogStyle::Shutdown()
+{
+	FSlateStyleRegistry::UnRegisterSlateStyle( *AddContentDialogStyleInstance );
+	ensure( AddContentDialogStyleInstance.IsUnique() );
+	AddContentDialogStyleInstance.Reset();
+}
+
+FName FAddContentDialogStyle::GetStyleSetName()
+{
+	static FName StyleSetName(TEXT("AddContentDialogStyle"));
+	return StyleSetName;
+}
+
+#define IMAGE_BRUSH( RelativePath, ... ) FSlateImageBrush( Style->RootToContentDir( RelativePath, TEXT(".png") ), __VA_ARGS__ )
+#define BOX_BRUSH( RelativePath, ... ) FSlateBoxBrush( Style->RootToContentDir( RelativePath, TEXT(".png") ), __VA_ARGS__ )
+#define BORDER_BRUSH( RelativePath, ... ) FSlateBorderBrush( Style->RootToContentDir( RelativePath, TEXT(".png") ), __VA_ARGS__ )
+#define TTF_FONT( RelativePath, ... ) FSlateFontInfo( Style->RootToContentDir( RelativePath, TEXT(".ttf") ), __VA_ARGS__ )
+#define OTF_FONT( RelativePath, ... ) FSlateFontInfo( Style->RootToContentDir( RelativePath, TEXT(".otf") ), __VA_ARGS__ )
+
+const FVector2D Icon16x16(16.0f, 16.0f);
+const FVector2D Icon20x20(20.0f, 20.0f);
+const FVector2D Icon40x40(40.0f, 40.0f);
+
+TSharedRef< FSlateStyleSet > FAddContentDialogStyle::Create()
+{
+	const FTextBlockStyle NormalText = FEditorStyle::GetWidgetStyle<FTextBlockStyle>("NormalText");
+
+	TSharedRef< FSlateStyleSet > Style = MakeShareable(new FSlateStyleSet("AddContentDialogStyle"));
+	Style->SetContentRoot(FPaths::EngineContentDir() / TEXT("Editor/Slate/AddContentDialog"));
+	
+	Style->Set("AddContentDialog.TabBackground", new BOX_BRUSH("TabBackground", 4 / 16.0f));
+	Style->Set("AddContentDialog.Splitter", new IMAGE_BRUSH("Splitter", FVector2D(8, 8)));
+
+	Style->Set("AddContentDialog.CategoryTab", FCheckBoxStyle()
+		.SetCheckBoxType(ESlateCheckBoxType::ToggleButton)
+		.SetUncheckedImage(BOX_BRUSH("TabInactive", 4 / 16.0f))
+		.SetUncheckedPressedImage(BOX_BRUSH("TabHovered", 4 / 16.0f))
+		.SetUncheckedHoveredImage(BOX_BRUSH("TabHovered", 4 / 16.0f))
+		.SetCheckedHoveredImage(BOX_BRUSH("TabActive", 4 / 16.0f))
+		.SetCheckedPressedImage(BOX_BRUSH("TabActive", 4 / 16.0f))
+		.SetCheckedImage(BOX_BRUSH("TabActive", 4 / 16.0f))		);
+
+	Style->Set("AddContentDialog.BlankButton", FButtonStyle()
+		.SetNormal(FSlateNoResource())
+		.SetHovered(FSlateNoResource())
+		.SetPressed(FSlateNoResource())		);
+
+	Style->Set("AddContentDialog.LeftArrow", new IMAGE_BRUSH("ArrowLeft", FVector2D(25.0f, 42.0f)));
+	Style->Set("AddContentDialog.RightArrow", new IMAGE_BRUSH("ArrowRight", FVector2D(25.0f, 42.0f)));
+
+	Style->Set("AddContentDialog.HeadingText", FTextBlockStyle(NormalText)
+		.SetFont(TTF_FONT("Fonts/Roboto-Regular", 14))
+		.SetShadowOffset(FVector2D(0, 1))
+		.SetShadowColorAndOpacity(FLinearColor(0, 0, 0, 0.9f)));
+
+	Style->Set("AddContentDialog.HeadingTextSmall", FTextBlockStyle(NormalText)
+		.SetFont(TTF_FONT("Fonts/Roboto-Regular", 12))
+		.SetShadowOffset(FVector2D(0, 1))
+		.SetShadowColorAndOpacity(FLinearColor(0, 0, 0, 0.9f)));
+
+	Style->Set("AddContentDialog.RemoveButton", FButtonStyle()
+		.SetNormal(IMAGE_BRUSH("Remove", FVector2D(16.0f, 16.0f)))
+		.SetHovered(IMAGE_BRUSH("RemoveHovered", FVector2D(16.0f, 16.0f)))
+		.SetPressed(IMAGE_BRUSH("RemoveHovered", FVector2D(16.0f, 16.0f)))		);
+
+	Style->Set("AddContentDialog.BlueprintFeatureCategory", new IMAGE_BRUSH("BlueprintFeature", FVector2D(32.0f, 32.0f)));
+	Style->Set("AddContentDialog.CodeFeatureCategory", new IMAGE_BRUSH("CodeFeature", FVector2D(32.0f, 32.0f)));
+	Style->Set("AddContentDialog.CategoryCategory", new FSlateNoResource());
+
+	return Style;
+}
+
+#undef IMAGE_BRUSH
+#undef BOX_BRUSH
+#undef BORDER_BRUSH
+#undef TTF_FONT
+#undef OTF_FONT
+
+void FAddContentDialogStyle::ReloadTextures()
+{
+	FSlateApplication::Get().GetRenderer()->ReloadTextureResources();
+}
+
+const ISlateStyle& FAddContentDialogStyle::Get()
+{
+	return *AddContentDialogStyleInstance;
+}
diff --git a/Engine/Source/Editor/AddContentDialog/Private/AddContentDialogStyle.h b/Engine/Source/Editor/AddContentDialog/Private/AddContentDialogStyle.h
new file mode 100644
--- /dev/null
+++ b/Engine/Source/Editor/AddContentDialog/Private/AddContentDialogStyle.h
@@ -0,0 +1,31 @@
+// Copyright 1998-2014 Epic Games, Inc. All Rights Reserved.
+
+#pragma once
+
+#include "SlateBasics.h"
+
+/** Manages the style which provides resources for the add content dialog. */
+class FAddContentDialogStyle
+{
+public:
+
+	static void Initialize();
+
+	static void Shutdown();
+
+	/** reloads textures used by slate renderer */
+	static void ReloadTextures();
+
+	/** @return The Slate style set for the Shooter game */
+	static const ISlateStyle& Get();
+
+	static FName GetStyleSetName();
+
+private:
+
+	static TSharedRef< class FSlateStyleSet > Create();
+
+private:
+
+	static TSharedPtr< class FSlateStyleSet > AddContentDialogStyleInstance;
+};
\ No newline at end of file
diff --git a/Engine/Source/Editor/AddContentDialog/Private/ContentSourceDragDropOp.cpp b/Engine/Source/Editor/AddContentDialog/Private/ContentSourceDragDropOp.cpp
new file mode 100644
--- /dev/null
+++ b/Engine/Source/Editor/AddContentDialog/Private/ContentSourceDragDropOp.cpp
@@ -0,0 +1,27 @@
+// Copyright 1998-2014 Epic Games, Inc. All Rights Reserved.
+
+#include "AddContentDialogPCH.h"
+
+TSharedRef<FContentSourceDragDropOp> FContentSourceDragDropOp::CreateShared(TSharedPtr<FContentSourceViewModel> InContentSource)
+{
+	TSharedPtr<FContentSourceDragDropOp> DragDropOp = MakeShareable(new FContentSourceDragDropOp(InContentSource));
+	DragDropOp->MouseCursor = EMouseCursor::GrabHandClosed;
+	DragDropOp->Construct();
+	return DragDropOp.ToSharedRef();
+}
+
+FContentSourceDragDropOp::FContentSourceDragDropOp(TSharedPtr<FContentSourceViewModel> InContentSource)
+{
+	ContentSource = InContentSource;
+}
+
+TSharedPtr<SWidget> FContentSourceDragDropOp::GetDefaultDecorator() const
+{
+	return SNew(SImage)
+		.Image(ContentSource->GetIconBrush().Get());
+}
+
+TSharedPtr<FContentSourceViewModel> FContentSourceDragDropOp::GetContentSource()
+{
+	return ContentSource;
+}
\ No newline at end of file
diff --git a/Engine/Source/Editor/AddContentDialog/Private/ContentSourceDragDropOp.h b/Engine/Source/Editor/AddContentDialog/Private/ContentSourceDragDropOp.h
new file mode 100644
--- /dev/null
+++ b/Engine/Source/Editor/AddContentDialog/Private/ContentSourceDragDropOp.h
@@ -0,0 +1,30 @@
+// Copyright 1998 - 2014 Epic Games, Inc.All Rights Reserved.
+
+#pragma once
+
+#include "DecoratedDragDropOp.h"
+
+class FContentSourceViewModel;
+
+/** A drag drop operation for dragging and dropping FContentSourceViewModels. */
+class FContentSourceDragDropOp : public FDecoratedDragDropOp
+{
+public:
+	DRAG_DROP_OPERATOR_TYPE(FContentSourceDragDropOp, FDecoratedDragDropOp)
+
+	/** Creates and constructs a shared references to a FContentSourceDragDrop.
+		@param InContentSource - The view model for the content source being dragged and dropped */
+	static TSharedRef<FContentSourceDragDropOp> CreateShared(TSharedPtr<FContentSourceViewModel> InContentSource);
+
+	virtual TSharedPtr<SWidget> GetDefaultDecorator() const override;
+
+	/** Gets the view model for the content source being dragged and dropped. */
+	TSharedPtr<FContentSourceViewModel> GetContentSource();
+
+private:
+	FContentSourceDragDropOp(TSharedPtr<FContentSourceViewModel> InContentSource);
+
+private:
+	/** The view model for the content source being dragged and dropped. */
+	TSharedPtr<FContentSourceViewModel> ContentSource;
+};
\ No newline at end of file
diff --git a/Engine/Source/Editor/AddContentDialog/Private/ContentSourceProviderManager.cpp b/Engine/Source/Editor/AddContentDialog/Private/ContentSourceProviderManager.cpp
new file mode 100644
--- /dev/null
+++ b/Engine/Source/Editor/AddContentDialog/Private/ContentSourceProviderManager.cpp
@@ -0,0 +1,13 @@
+// Copyright 1998-2014 Epic Games, Inc. All Rights Reserved.
+
+#include "AddContentDialogPCH.h"
+
+void FContentSourceProviderManager::RegisterContentSourceProvider(TSharedRef<IContentSourceProvider> ContentSourceProvider)
+{
+	ContentSourceProviders.Add(ContentSourceProvider);
+}
+
+const TArray<TSharedRef<IContentSourceProvider>>* FContentSourceProviderManager::GetContentSourceProviders()
+{
+	return &ContentSourceProviders;
+}
\ No newline at end of file
diff --git a/Engine/Source/Editor/AddContentDialog/Private/ContentSourceProviderManager.h b/Engine/Source/Editor/AddContentDialog/Private/ContentSourceProviderManager.h
new file mode 100644
--- /dev/null
+++ b/Engine/Source/Editor/AddContentDialog/Private/ContentSourceProviderManager.h
@@ -0,0 +1,20 @@
+// Copyright 1998-2014 Epic Games, Inc. All Rights Reserved.
+
+#pragma once
+
+class IContentSourceProvider;
+
+/** Provides methods for registering and getting registered content source providers. */
+class FContentSourceProviderManager
+{
+public:
+	/** Registers a content source provider. */
+	void RegisterContentSourceProvider(TSharedRef<IContentSourceProvider> ContentSourceProvider);
+
+	/** Gets and array of the currently registered content source providers. */
+	const TArray<TSharedRef<IContentSourceProvider>>* GetContentSourceProviders();
+
+private:
+	/** The registered content source providers. */
+	TArray<TSharedRef<IContentSourceProvider>> ContentSourceProviders;
+};
\ No newline at end of file
diff --git a/Engine/Source/Editor/AddContentDialog/Private/ContentSourceProviders/FeaturePack/FeaturePackContentSource.cpp b/Engine/Source/Editor/AddContentDialog/Private/ContentSourceProviders/FeaturePack/FeaturePackContentSource.cpp
new file mode 100644
--- /dev/null
+++ b/Engine/Source/Editor/AddContentDialog/Private/ContentSourceProviders/FeaturePack/FeaturePackContentSource.cpp
@@ -0,0 +1,113 @@
+// Copyright 1998-2014 Epic Games, Inc. All Rights Reserved.
+
+#include "AddContentDialogPCH.h"
+#include "FeaturePackContentSource.h"
+
+#include "AssetToolsModule.h"
+#include "IPlatformFilePak.h"
+
+FFeaturePackContentSource::FFeaturePackContentSource(FString InFeaturePackPath)
+{
+	FeaturePackPath = InFeaturePackPath;
+
+	// Create a pak platform file and mount the feature pack file.
+	FPakPlatformFile PakPlatformFile;
+	FString CommandLine;
+	PakPlatformFile.Initialize(&FPlatformFileManager::Get().GetPlatformFile(), TEXT(""));
+	FString MountPoint = "root:/";
+	PakPlatformFile.Mount(*InFeaturePackPath, 0, *MountPoint);
+
+	// Gets the manifest file as a JSon string
+	TArray<uint8> ManifestBuffer;
+	LoadPakFileToBuffer(PakPlatformFile, FPaths::Combine(*MountPoint, TEXT("manifest.json")), ManifestBuffer);
+	FString ManifestString;
+	FFileHelper::BufferToString(ManifestString, ManifestBuffer.GetData(), ManifestBuffer.Num());
+
+	// Populate text fields from the manifest.
+	TSharedPtr<FJsonObject> ManifestObject;
+	TSharedRef<TJsonReader<>> ManifestReader = TJsonReaderFactory<>::Create(ManifestString);
+	FJsonSerializer::Deserialize(ManifestReader, ManifestObject);
+	for (TSharedPtr<FJsonValue> NameValue : ManifestObject->GetArrayField("Name"))
+	{
+		TSharedPtr<FJsonObject> LocalizedNameObject = NameValue->AsObject();
+		LocalizedNames.Add(FLocalizedText(
+			LocalizedNameObject->GetStringField("Language"),
+			FText::FromString(LocalizedNameObject->GetStringField("Text"))));
+	}
+
+	for (TSharedPtr<FJsonValue> DescriptionValue : ManifestObject->GetArrayField("Description"))
+	{
+		TSharedPtr<FJsonObject> LocalizedDescriptionObject = DescriptionValue->AsObject();
+		LocalizedDescriptions.Add(FLocalizedText(
+			LocalizedDescriptionObject->GetStringField("Language"),
+			FText::FromString(LocalizedDescriptionObject->GetStringField("Text"))));
+	}
+
+	FString CategoryString = ManifestObject->GetStringField("Category");
+	if (CategoryString == "CodeFeature")
+	{
+		Category = EContentSourceCategory::CodeFeature;
+	}
+	else if (CategoryString == "BlueprintFeature")
+	{
+		Category = EContentSourceCategory::BlueprintFeature;
+	}
+	else
+	{
+		Category = EContentSourceCategory::Unknown;
+	}
+
+	// Load image data
+	FString IconFilename = ManifestObject->GetStringField("Thumbnail");
+	TSharedPtr<TArray<uint8>> IconImageData = MakeShareable(new TArray<uint8>());
+	LoadPakFileToBuffer(PakPlatformFile, FPaths::Combine(*MountPoint, TEXT("Media"), *IconFilename), *IconImageData);
+	IconData = MakeShareable(new FImageData(IconFilename, IconImageData));
+
+	const TArray<TSharedPtr<FJsonValue>> ScreenshotFilenameArray = ManifestObject->GetArrayField("Screenshots");
+	for (const TSharedPtr<FJsonValue> ScreenshotFilename : ScreenshotFilenameArray)
+	{
+		TSharedPtr<TArray<uint8>> SingleScreenshotData = MakeShareable(new TArray<uint8>);
+		LoadPakFileToBuffer(PakPlatformFile, FPaths::Combine(*MountPoint, TEXT("Media"), *ScreenshotFilename->AsString()), *SingleScreenshotData);
+		ScreenshotData.Add(MakeShareable(new FImageData(ScreenshotFilename->AsString(), SingleScreenshotData)));
+	}
+}
+
+void FFeaturePackContentSource::LoadPakFileToBuffer(FPakPlatformFile& PakPlatformFile, FString Path, TArray<uint8>& Buffer)
+{
+	TSharedPtr<IFileHandle> FileHandle(PakPlatformFile.OpenRead(*Path));
+	Buffer.AddUninitialized(FileHandle->Size());
+	FileHandle->Read(Buffer.GetData(), FileHandle->Size());
+}
+
+TArray<FLocalizedText> FFeaturePackContentSource::GetLocalizedNames()
+{
+	return LocalizedNames;
+}
+
+TArray<FLocalizedText> FFeaturePackContentSource::GetLocalizedDescriptions()
+{
+	return LocalizedDescriptions;
+}
+
+EContentSourceCategory FFeaturePackContentSource::GetCategory()
+{
+	return Category;
+}
+
+TSharedPtr<FImageData> FFeaturePackContentSource::GetIconData()
+{
+	return IconData;
+}
+
+TArray<TSharedPtr<FImageData>> FFeaturePackContentSource::GetScreenshotData()
+{
+	return ScreenshotData;
+}
+
+void FFeaturePackContentSource::InstallToProject(FString InstallPath)
+{
+	FAssetToolsModule& AssetToolsModule = FModuleManager::Get().LoadModuleChecked<FAssetToolsModule>("AssetTools");
+	TArray<FString> AssetPaths;
+	AssetPaths.Add(FeaturePackPath);
+	AssetToolsModule.Get().ImportAssets(AssetPaths, InstallPath);
+}
\ No newline at end of file
diff --git a/Engine/Source/Editor/AddContentDialog/Private/ContentSourceProviders/FeaturePack/FeaturePackContentSource.h b/Engine/Source/Editor/AddContentDialog/Private/ContentSourceProviders/FeaturePack/FeaturePackContentSource.h
new file mode 100644
--- /dev/null
+++ b/Engine/Source/Editor/AddContentDialog/Private/ContentSourceProviders/FeaturePack/FeaturePackContentSource.h
@@ -0,0 +1,33 @@
+// Copyright 1998-2014 Epic Games, Inc. All Rights Reserved.
+
+#pragma once
+
+class FPakPlatformFile;
+
+/** A content source which represents a content upack. */
+class FFeaturePackContentSource : public IContentSource
+{
+public:
+	FFeaturePackContentSource(FString InFeaturePackPath);
+
+	virtual TArray<FLocalizedText> GetLocalizedNames() override;
+	virtual TArray<FLocalizedText> GetLocalizedDescriptions() override;
+	
+	virtual EContentSourceCategory GetCategory() override;
+
+	virtual TSharedPtr<FImageData> GetIconData() override;
+	virtual TArray<TSharedPtr<FImageData>> GetScreenshotData() override;
+
+	virtual void InstallToProject(FString InstallPath) override;
+
+private:
+	void LoadPakFileToBuffer(FPakPlatformFile& PakPlatformFile, FString Path, TArray<uint8>& Buffer);
+
+private:
+	FString FeaturePackPath;
+	TArray<FLocalizedText> LocalizedNames;
+	TArray<FLocalizedText> LocalizedDescriptions;
+	EContentSourceCategory Category;
+	TSharedPtr<FImageData> IconData;
+	TArray<TSharedPtr<FImageData>> ScreenshotData;
+};
\ No newline at end of file
diff --git a/Engine/Source/Editor/AddContentDialog/Private/ContentSourceProviders/FeaturePack/FeaturePackContentSourceProvider.cpp b/Engine/Source/Editor/AddContentDialog/Private/ContentSourceProviders/FeaturePack/FeaturePackContentSourceProvider.cpp
new file mode 100644
--- /dev/null
+++ b/Engine/Source/Editor/AddContentDialog/Private/ContentSourceProviders/FeaturePack/FeaturePackContentSourceProvider.cpp
@@ -0,0 +1,49 @@
+// Copyright 1998-2014 Epic Games, Inc. All Rights Reserved.
+
+#include "AddContentDialogPCH.h"
+
+#include "FeaturePackContentSourceProvider.h"
+#include "FeaturePackContentSource.h"
+
+class FFillArrayDirectoryVisitor : public IPlatformFile::FDirectoryVisitor
+{
+public:
+	virtual bool Visit(const TCHAR* FilenameOrDirectory, bool bIsDirectory) override
+	{
+		if (bIsDirectory)
+		{
+			Directories.Add(FilenameOrDirectory);
+		}
+		else
+		{
+			Files.Add(FilenameOrDirectory);
+		}
+		return true;
+	}
+
+	TArray<FString> Directories;
+	TArray<FString> Files;
+};
+
+FFeaturePackContentSourceProvider::FFeaturePackContentSourceProvider()
+{
+	//TODO: Add a directory watcher?
+	FString FeaturePackPath = FPaths::Combine(*FPaths::RootDir(), TEXT("FeaturePacks"), TEXT("Packed"));
+	IPlatformFile &PlatformFile = FPlatformFileManager::Get().GetPlatformFile();
+	FFillArrayDirectoryVisitor DirectoryVisitor;
+	PlatformFile.IterateDirectory(*FeaturePackPath, DirectoryVisitor);
+	for (auto FeaturePackFile : DirectoryVisitor.Files)
+	{
+		ContentSources.Add(MakeShareable(new FFeaturePackContentSource(FeaturePackFile)));
+	}
+}
+
+const TArray<TSharedRef<IContentSource>> FFeaturePackContentSourceProvider::GetContentSources()
+{
+	return ContentSources;
+}
+
+void FFeaturePackContentSourceProvider::SetContentSourcesChanged(FOnContentSourcesChanged OnContentSourcesChangedIn)
+{
+	OnContentSourcesChanged = OnContentSourcesChangedIn;
+}
\ No newline at end of file
diff --git a/Engine/Source/Editor/AddContentDialog/Private/ContentSourceProviders/FeaturePack/FeaturePackContentSourceProvider.h b/Engine/Source/Editor/AddContentDialog/Private/ContentSourceProviders/FeaturePack/FeaturePackContentSourceProvider.h
new file mode 100644
--- /dev/null
+++ b/Engine/Source/Editor/AddContentDialog/Private/ContentSourceProviders/FeaturePack/FeaturePackContentSourceProvider.h
@@ -0,0 +1,16 @@
+// Copyright 1998-2014 Epic Games, Inc. All Rights Reserved.
+
+#pragma once
+
+/** A content source provider for available content upacks. */
+class FFeaturePackContentSourceProvider : public IContentSourceProvider
+{
+public:
+	FFeaturePackContentSourceProvider();
+	virtual const TArray<TSharedRef<IContentSource>> GetContentSources() override;
+	virtual void SetContentSourcesChanged(FOnContentSourcesChanged OnContentSourcesChangedIn) override;
+
+private:
+	FOnContentSourcesChanged OnContentSourcesChanged;
+	TArray<TSharedRef<IContentSource>> ContentSources;
+};
\ No newline at end of file
diff --git a/Engine/Source/Editor/AddContentDialog/Private/IContentSource.h b/Engine/Source/Editor/AddContentDialog/Private/IContentSource.h
new file mode 100644
--- /dev/null
+++ b/Engine/Source/Editor/AddContentDialog/Private/IContentSource.h
@@ -0,0 +1,98 @@
+// Copyright 1998 - 2014 Epic Games, Inc.All Rights Reserved.
+
+#pragma once
+
+/** Defines categories for content sources. */
+enum class EContentSourceCategory
+{
+	BlueprintFeature,
+	CodeFeature,
+	Unknown
+};
+
+/** Represents raw binary image data in png format. */
+class FImageData
+{
+public:
+	/** Creates a new FImageData
+		@param InName - The name which identifies the image represented by this object 
+		@param InData - The raw binary image data in png format */
+	FImageData(FString InName, TSharedPtr <TArray<uint8>> InData)
+	{
+		Name = InName;
+		Data = InData;
+	}
+
+	/** Gets the name which identifies the image represented by this obejct. */
+	const FString GetName() const
+	{
+		return Name;
+	}
+
+	/** Gets the raw binary image data in PNG format. */
+	const TSharedPtr<TArray<uint8>> GetData() const
+	{
+		return Data;
+	}
+
+private:
+	FString Name;
+	TSharedPtr<TArray<uint8>> Data;
+};
+
+/** Represents a pieces of localized text. */
+class FLocalizedText
+{
+public:
+	FLocalizedText()
+	{
+	}
+
+	/** Creates a new FLocalizedText
+		@param InTwoLetterLanguage - The iso 2-letter language specifier.
+		@param InText - The text in the language specified */
+	FLocalizedText(FString InTwoLetterLanguage, FText InText)
+	{
+		TwoLetterLanguage = InTwoLetterLanguage;
+		Text = InText;
+	}
+
+	/** Gets the iso 2-letter language specifier for this text. */
+	FString GetTwoLetterLanguage() const
+	{
+		return TwoLetterLanguage;
+	}
+
+	/** Gets the text in the language specified. */
+	FText GetText() const
+	{
+		return Text;
+	}
+
+private:
+	FString TwoLetterLanguage;
+	FText Text;
+};
+
+/** Defines a source of content to be used with the FAddContentDialog */
+class IContentSource
+{
+public:
+	/** Gets the name of the content source as an array of localized strings. */
+	virtual TArray<FLocalizedText> GetLocalizedNames() = 0;
+
+	/** Gets the description of the content source as an array or localized strings. */
+	virtual TArray<FLocalizedText> GetLocalizedDescriptions() = 0;
+
+	/** Gets the category for the content source. */
+	virtual EContentSourceCategory GetCategory() = 0;
+
+	/** Gets the image data for the icon which should represent the content source in the UI. */
+	virtual TSharedPtr<FImageData> GetIconData() = 0;
+
+	/** Gets an array or image data for screenshots for the content source. */
+	virtual TArray<TSharedPtr<FImageData>> GetScreenshotData() = 0;
+
+	/** Installs the content in the content source to the specific path. */
+	virtual void InstallToProject(FString InstallPath) = 0;
+};
diff --git a/Engine/Source/Editor/AddContentDialog/Private/IContentSourceProvider.h b/Engine/Source/Editor/AddContentDialog/Private/IContentSourceProvider.h
new file mode 100644
--- /dev/null
+++ b/Engine/Source/Editor/AddContentDialog/Private/IContentSourceProvider.h
@@ -0,0 +1,19 @@
+// Copyright 1998-2014 Epic Games, Inc. All Rights Reserved.
+
+#pragma once
+
+class IContentSource;
+
+/** Defines methods for an object which provides IContentSource objects for use with the SAddContentDialog. */
+class IContentSourceProvider
+{
+public:
+	DECLARE_DELEGATE(FOnContentSourcesChanged);
+
+public:
+	/** Gets the available content sources. */
+	virtual const TArray<TSharedRef<IContentSource>> GetContentSources() = 0;
+
+	/** Sets the delegate which will be executed when the avaialbe content sources change */
+	virtual void SetContentSourcesChanged(FOnContentSourcesChanged OnContentSourcesChangedIn) = 0;
+};
\ No newline at end of file
diff --git a/Engine/Source/Editor/AddContentDialog/Private/SAddContentDialog.cpp b/Engine/Source/Editor/AddContentDialog/Private/SAddContentDialog.cpp
new file mode 100644
--- /dev/null
+++ b/Engine/Source/Editor/AddContentDialog/Private/SAddContentDialog.cpp
@@ -0,0 +1,63 @@
+// Copyright 1998-2014 Epic Games, Inc. All Rights Reserved.
+
+#include "AddContentDialogPCH.h"
+
+#define LOCTEXT_NAMESPACE "AddContentDialog"
+
+void SAddContentDialog::Construct(const FArguments& InArgs)
+{
+	SAssignNew(AddContentWidget, SAddContentWidget);
+
+	SWindow::Construct(SWindow::FArguments()
+		.Title(LOCTEXT("AddContentDialogTitle", "Add Content to the Project"))
+		.SizingRule(ESizingRule::UserSized)
+		.ClientSize(FVector2D(900, 700))
+		.SupportsMinimize(false)
+		.SupportsMaximize(false)
+		[
+			SNew(SBorder)
+			.BorderImage(FEditorStyle::GetBrush("ToolPanel.GroupBorder"))
+			.Padding(FMargin(15))
+			[
+				SNew(SVerticalBox)
+
+				// Add content widget.
+				+ SVerticalBox::Slot()
+				.Padding(FMargin(0, 0, 0, 15))
+				[
+					AddContentWidget.ToSharedRef()
+				]
+
+				// Add button.
+				+ SVerticalBox::Slot()
+				.AutoHeight()
+				.HAlign(EHorizontalAlignment::HAlign_Right)
+				[
+					SNew(SButton)
+					.OnClicked(this, &SAddContentDialog::AddButtonClicked)
+					.IsEnabled_Lambda([this]()->bool
+					{
+						return AddContentWidget->GetContentSourcesToAdd()->Num() > 0;
+					})
+					[
+						SNew(STextBlock)
+						.Text(LOCTEXT("AddSelectedContentToProject", "Add the selected content to the project"))
+					]
+				]
+			]
+		]);
+}
+
+FReply SAddContentDialog::AddButtonClicked()
+{
+	for (TSharedPtr<IContentSource> ContentSource : *AddContentWidget->GetContentSourcesToAdd())
+	{
+		ContentSource->InstallToProject("/Game");
+	}
+	RequestDestroyWindow();
+	return FReply::Handled();
+}
+
+SAddContentDialog::~SAddContentDialog()
+{
+}
\ No newline at end of file
diff --git a/Engine/Source/Editor/AddContentDialog/Private/SAddContentDialog.h b/Engine/Source/Editor/AddContentDialog/Private/SAddContentDialog.h
new file mode 100644
--- /dev/null
+++ b/Engine/Source/Editor/AddContentDialog/Private/SAddContentDialog.h
@@ -0,0 +1,25 @@
+// Copyright 1998-2014 Epic Games, Inc. All Rights Reserved.
+
+#pragma once
+
+/** A window which allows the user to select additional content to add to the currently loaded project. */
+class SAddContentDialog : public SWindow
+{
+public:
+	SLATE_BEGIN_ARGS(SWindow)
+	{}
+
+	SLATE_END_ARGS()
+
+	~SAddContentDialog();
+
+	void Construct(const FArguments& InArgs);
+
+private:
+	/** Handles the add content to project button being clicked. */
+	FReply AddButtonClicked();
+
+private:
+	/** The widget representing available content and which content the user has selected. */
+	TSharedPtr<SAddContentWidget> AddContentWidget;
+};
\ No newline at end of file
diff --git a/Engine/Source/Editor/AddContentDialog/Private/SAddContentWidget.cpp b/Engine/Source/Editor/AddContentDialog/Private/SAddContentWidget.cpp
new file mode 100644
--- /dev/null
+++ b/Engine/Source/Editor/AddContentDialog/Private/SAddContentWidget.cpp
@@ -0,0 +1,514 @@
+// Copyright 1998-2014 Epic Games, Inc. All Rights Reserved.
+
+#include "AddContentDialogPCH.h"
+
+#include "SSearchBox.h"
+
+#define LOCTEXT_NAMESPACE "AddContentDialog"
+
+class SContentSourceDragDropHandler : public SCompoundWidget
+{
+public:
+	SLATE_BEGIN_ARGS(SContentSourceDragDropHandler) {}
+	SLATE_DEFAULT_SLOT(FArguments, Content)
+	SLATE_ARGUMENT(TWeakPtr<FAddContentWidgetViewModel>, ViewModel)
+	SLATE_END_ARGS()
+
+	void Construct(const FArguments& InArgs)
+	{
+		ViewModel = InArgs._ViewModel;
+		bIsDragOver = false;
+
+		this->ChildSlot
+		[
+			SNew(SBorder)
+			.BorderImage(FEditorStyle::GetBrush("WhiteBrush"))
+			.BorderBackgroundColor(this, &SContentSourceDragDropHandler::GetBackgroundColor)
+			[
+				InArgs._Content.Widget
+			]
+		];
+	}
+
+	~SContentSourceDragDropHandler() {}
+
+	FReply OnDrop(const FGeometry& MyGeometry, const FDragDropEvent& DragDropEvent) override
+	{
+		bIsDragOver = false;
+
+		TSharedPtr<FAddContentWidgetViewModel> PinnedViewModel = ViewModel.Pin();
+		TSharedPtr<FContentSourceDragDropOp> ContentSourceDragDropOp = DragDropEvent.GetOperationAs<FContentSourceDragDropOp>();
+		if (PinnedViewModel.IsValid() && ContentSourceDragDropOp.IsValid())
+		{
+			PinnedViewModel->AppendToAddList(ContentSourceDragDropOp->GetContentSource());
+			return FReply::Handled();
+		}
+		return FReply::Unhandled();
+	}
+
+	virtual void OnDragEnter(const FGeometry& MyGeometry, const FDragDropEvent& DragDropEvent) override
+	{
+		bIsDragOver = true;
+	}
+
+	virtual void OnDragLeave(const FDragDropEvent& DragDropEvent) override
+	{
+		bIsDragOver = false;
+	}
+
+private:
+	FSlateColor GetBackgroundColor() const
+	{
+		if (bIsDragOver)
+		{
+			return FLinearColor(1.0f, 0.6f, 0.1f, 0.9f);
+		}
+		else
+		{
+			return FSlateApplication::Get().IsDragDropping() && FSlateApplication::Get().GetDragDroppingContent()->IsOfType<FContentSourceDragDropOp>() ?
+				FLinearColor(0.1f, 0.1f, 0.1f, 0.9f) : FLinearColor::Transparent;
+		}
+	}
+
+private:
+	TWeakPtr<FAddContentWidgetViewModel> ViewModel;
+	bool bIsDragOver;
+};
+
+void SAddContentWidget::Construct(const FArguments& InArgs)
+{
+	ViewModel = FAddContentWidgetViewModel::CreateShared();
+	ViewModel->SetOnCategoriesChanged(FAddContentWidgetViewModel::FOnCategoriesChanged::CreateSP(
+		this, &SAddContentWidget::CategoriesChanged));
+	ViewModel->SetOnContentSourcesChanged(FAddContentWidgetViewModel::FOnContentSourcesChanged::CreateSP(
+		this, &SAddContentWidget::ContentSourcesChanged));
+	ViewModel->SetOnSelectedContentSourceChanged(FAddContentWidgetViewModel::FOnSelectedContentSourceChanged::CreateSP(
+		this, &SAddContentWidget::SelectedContentSourceChanged));
+	ViewModel->SetOnAddListChanged(FAddContentWidgetViewModel::FOnAddListChanged::CreateSP(
+		this, &SAddContentWidget::AddListChanged));
+
+	ChildSlot
+	[
+		SNew(SVerticalBox)
+
+		// Tab Buttons
+		+ SVerticalBox::Slot()
+		.AutoHeight()
+		.Padding(FMargin(10, 0))
+		[
+			SAssignNew(CategoryTabsContainer, SBox)
+			[
+				CreateCategoryTabs()
+			]
+		]
+
+		// Content Source Tab Page
+		+ SVerticalBox::Slot()
+		.FillHeight(3.0f)
+		.Padding(FMargin(0, 0, 0, 10))
+		[
+			SNew(SBorder)
+			.BorderImage(FAddContentDialogStyle::Get().GetBrush("AddContentDialog.TabBackground"))
+			.Padding(FMargin(15))
+			[
+				SNew(SHorizontalBox)
+
+				// Content Source Tiles
+				+ SHorizontalBox::Slot()
+				[
+					SNew(SVerticalBox)
+
+					// Content Source Filter
+					+ SVerticalBox::Slot()
+					.AutoHeight()
+					.Padding(FMargin(0, 0, 0, 5))
+					[
+						SNew(SSearchBox)
+						.OnTextChanged(this, &SAddContentWidget::SearchTextChanged)
+					]
+
+					// Content Source Tile View
+					+ SVerticalBox::Slot()
+					[
+						CreateContentSourceTileView()
+					]
+				]
+
+				// Splitter
+				+ SHorizontalBox::Slot()
+				.Padding(FMargin(10, 0))
+				.AutoWidth()
+				[
+					SNew(SBox)
+					.WidthOverride(2)
+					[
+						SNew(SImage)
+						.Image(FAddContentDialogStyle::Get().GetBrush("AddContentDialog.Splitter"))
+					]
+				]
+
+				// Content Source Details
+				+ SHorizontalBox::Slot()
+				[
+					SAssignNew(ContentSourceDetailContainer, SBox)
+					[
+						CreateContentSourceDetail(ViewModel->GetSelectedContentSource())
+					]
+				]
+			]
+		]
+
+		// List of selected content
+		+ SVerticalBox::Slot()
+		.FillHeight(1.0f)
+		[
+			CreateAddListView()
+		]
+	];
+}
+
+const TArray<TSharedPtr<IContentSource>>* SAddContentWidget::GetContentSourcesToAdd()
+{
+	return &ContentSourcesToAdd;
+}
+
+TSharedRef<SWidget> SAddContentWidget::CreateCategoryTabs()
+{
+	TSharedRef<SHorizontalBox> TabBox = SNew(SHorizontalBox);
+	for (auto Category : *ViewModel->GetCategories())
+	{
+		TabBox->AddSlot()
+		.Padding(FMargin(0, 0, 5, 0))
+		.AutoWidth()
+		[
+			SNew(SCheckBox)
+			.Style(FAddContentDialogStyle::Get(), "AddContentDialog.CategoryTab")
+			.OnCheckStateChanged(this, &SAddContentWidget::CategoryCheckBoxCheckStateChanged, Category)
+			.IsChecked(this, &SAddContentWidget::GetCategoryCheckBoxCheckState, Category)
+			.Padding(FMargin(5))
+			[
+				SNew(SHorizontalBox)
+
+				+ SHorizontalBox::Slot()
+				.AutoWidth()
+				.VAlign(VAlign_Center)
+				.Padding(FMargin(0, 0, 5, 0))
+				[
+					SNew(SImage)
+					.Image(Category.GetIconBrush())
+				]
+
+				+ SHorizontalBox::Slot()
+				.AutoWidth()
+				.VAlign(VAlign_Center)
+				[
+					SNew(STextBlock)
+					.TextStyle(FAddContentDialogStyle::Get(), "AddContentDialog.HeadingText")
+					.Text(Category.GetText())
+				]
+			]
+		];
+	}
+	return TabBox;
+}
+
+TSharedRef<SWidget> SAddContentWidget::CreateContentSourceTileView()
+{
+	SAssignNew(ContentSourceTileView, STileView<TSharedPtr<FContentSourceViewModel>>)
+	.ListItemsSource(ViewModel->GetContentSources())
+	.OnGenerateTile(this, &SAddContentWidget::CreateContentSourceIconTile)
+	.OnSelectionChanged(this, &SAddContentWidget::ContentSourceTileViewSelectionChanged)
+	.ItemWidth(70)
+	.ItemHeight(100);
+	ContentSourceTileView->SetSelection(ViewModel->GetSelectedContentSource(), ESelectInfo::Direct);
+	return ContentSourceTileView.ToSharedRef();
+}
+
+TSharedRef<ITableRow> SAddContentWidget::CreateContentSourceIconTile(TSharedPtr<FContentSourceViewModel> ContentSource, const TSharedRef<STableViewBase>& OwnerTable)
+{
+	return SNew(STableRow< TSharedPtr<FString> >, OwnerTable)
+	.Cursor(EMouseCursor::GrabHand)
+	.OnDragDetected(this, &SAddContentWidget::ContentSourceTileDragDetected)
+	[
+		SNew(SVerticalBox)
+
+		+ SVerticalBox::Slot()
+		.HAlign(EHorizontalAlignment::HAlign_Center)
+		.AutoHeight()
+		[
+			SNew(SImage)
+			.Image(ContentSource->GetIconBrush().Get())
+		]
+
+		+ SVerticalBox::Slot()
+		.HAlign(HAlign_Center)
+		.AutoHeight()
+		[
+			SNew(STextBlock)
+			.Text(ContentSource->GetName())
+			.AutoWrapText(true)
+		]
+	];
+}
+
+TSharedRef<SWidget> SAddContentWidget::CreateContentSourceDetail(TSharedPtr<FContentSourceViewModel> ContentSource)
+{
+	TSharedRef<SScrollBox> Box = SNew(SScrollBox);
+	if (ContentSource.IsValid())
+	{
+		Box->AddSlot()
+		.Padding(FMargin(0, 0, 0, 5))
+		.HAlign(EHorizontalAlignment::HAlign_Left)
+		[
+			CreateScreenshotCarousel(ContentSource)
+		];
+
+		Box->AddSlot()
+		.Padding(FMargin(0, 0, 0, 5))
+		[
+			SNew(STextBlock)
+			.TextStyle(FAddContentDialogStyle::Get(), "AddContentDialog.HeadingText")
+			.Text(ContentSource->GetName())
+			.AutoWrapText(true)
+		];
+
+		Box->AddSlot()
+		.Padding(FMargin(0, 0, 0, 5))
+		[
+			SNew(STextBlock)
+			.Text(ContentSource->GetDescription())
+			.AutoWrapText(true)
+		];
+	}
+	return Box;
+}
+
+TSharedRef<SWidget> SAddContentWidget::CreateScreenshotCarousel(TSharedPtr<FContentSourceViewModel> ContentSource)
+{
+	// Previous button and image.
+	TSharedRef<SButton> PreviousButton = 
+		SNew(SButton)
+		.ButtonStyle(FAddContentDialogStyle::Get(), "AddContentDialog.BlankButton")
+		.Visibility(ContentSource->GetScreenshotBrushes()->Num() > 1 ? EVisibility::Visible : EVisibility::Hidden)
+		.OnClicked(this, &SAddContentWidget::PreviousScreenshotButtonClicked, ContentSource)
+		.HAlign(EHorizontalAlignment::HAlign_Left)
+		.VAlign(EVerticalAlignment::VAlign_Center);
+
+	TSharedRef<SImage> PreviousImage = 
+		SNew(SImage)
+		.Image(FAddContentDialogStyle::Get().GetBrush("AddContentDialog.LeftArrow"))
+		.Visibility_Lambda([PreviousButton]()->EVisibility
+		{
+			return PreviousButton->IsHovered() ? EVisibility::Visible : EVisibility::Hidden;
+		});
+
+	PreviousButton->SetContent(PreviousImage);
+
+	// Next button and image.
+	TSharedRef<SButton> NextButton = 
+		SNew(SButton)
+		.ButtonStyle(FAddContentDialogStyle::Get(), "AddContentDialog.BlankButton")
+		.Visibility(ContentSource->GetScreenshotBrushes()->Num() > 1 ? EVisibility::Visible : EVisibility::Hidden)
+		.OnClicked(this, &SAddContentWidget::NextScreenshotButtonClicked, ContentSource)
+		.HAlign(EHorizontalAlignment::HAlign_Right)
+		.VAlign(EVerticalAlignment::VAlign_Center);
+
+	TSharedRef<SImage> NextImage = 
+		SNew(SImage)
+		.Image(FAddContentDialogStyle::Get().GetBrush("AddContentDialog.RightArrow"))
+		.Visibility_Lambda([NextButton]()->EVisibility
+		{
+			return NextButton->IsHovered() ? EVisibility::Visible : EVisibility::Hidden;
+		});
+
+	NextButton->SetContent(NextImage);
+
+	// Image and button overlay.
+	return SNew(SOverlay)
+
+	+ SOverlay::Slot()
+	.HAlign(HAlign_Center)
+	.VAlign(VAlign_Center)
+	[
+		SNew(SImage)
+		.Image(ContentSource.ToSharedRef(), &FContentSourceViewModel::GetSelectedScreenshotBrush)
+	]
+
+	+ SOverlay::Slot()
+	.Padding(0)
+	[
+		SNew(SHorizontalBox)
+		+ SHorizontalBox::Slot()
+		[
+			PreviousButton
+		]
+		+ SHorizontalBox::Slot()
+		[
+			NextButton
+		]
+	];
+}
+
+TSharedRef<SWidget> SAddContentWidget::CreateAddListView()
+{
+	return SNew(SVerticalBox)
+
+	+ SVerticalBox::Slot()
+	.AutoHeight()
+	.Padding(FMargin(0, 0, 0, 5))
+	[
+		SNew(STextBlock)
+		.TextStyle(FAddContentDialogStyle::Get(), "AddContentDialog.HeadingText")
+		.Text(LOCTEXT("SelectedContentToAdd", "Selected Content"))
+	]
+
+	+ SVerticalBox::Slot()
+	[
+		SNew(SOverlay)
+		
+		+ SOverlay::Slot()
+		[
+			SAssignNew(AddListView, SListView<TSharedPtr<FContentSourceViewModel>>)
+			.ListItemsSource(ViewModel->GetAddList())
+			.OnGenerateRow(this, &SAddContentWidget::CreateAddContentSourceRow)
+		]
+
+		+ SOverlay::Slot()
+		[
+			SNew(SContentSourceDragDropHandler)
+			.ViewModel(TWeakPtr<FAddContentWidgetViewModel>(ViewModel))
+			.Visibility_Lambda([this]()->EVisibility
+			{
+				if (ViewModel->GetAddList()->Num() == 0)
+				{
+					return EVisibility::Visible;
+				}
+				else
+				{
+					return FSlateApplication::Get().IsDragDropping() && FSlateApplication::Get().GetDragDroppingContent()->IsOfType<FContentSourceDragDropOp>() ?
+						EVisibility::Visible : EVisibility::Collapsed;
+				}
+			})
+			[
+				SNew(SBorder)
+				.HAlign(EHorizontalAlignment::HAlign_Center)
+				.VAlign(EVerticalAlignment::VAlign_Center)
+				[
+					SNew(STextBlock)
+					.TextStyle(FAddContentDialogStyle::Get(), "AddContentDialog.HeadingTextSmall")
+					.Text(LOCTEXT("DropContent", "+ Drop Content Here"))
+				]
+			]
+		]
+	];
+}
+
+TSharedRef<ITableRow> SAddContentWidget::CreateAddContentSourceRow(TSharedPtr<FContentSourceViewModel> ContentSource, const TSharedRef<STableViewBase>& OwnerTable)
+{
+	return SNew(STableRow< TSharedPtr<FString> >, OwnerTable)
+	.Padding(5)
+	[
+		SNew(SHorizontalBox)
+
+		+ SHorizontalBox::Slot()
+		.AutoWidth()
+		.Padding(FMargin(0, 0, 5, 0))
+		[
+			SNew(STextBlock)
+			.Text_Lambda([ContentSource]()->FText
+			{
+				return FText::Format(LOCTEXT("CategoryNameFormat", "{0} - {1}"), ContentSource->GetCategory().GetText(), ContentSource->GetName());
+			})
+		]
+
+		+ SHorizontalBox::Slot()
+		.AutoWidth()
+		[
+			SNew(SButton)
+			.ButtonStyle(FAddContentDialogStyle::Get(), "AddContentDialog.RemoveButton")
+			.ToolTipText(LOCTEXT("RemoteToolTip", "Remove"))
+			.OnClicked(this, &SAddContentWidget::RemoveAddedContentSourceClicked, ContentSource)
+		]
+	];
+}
+
+void SAddContentWidget::CategoryCheckBoxCheckStateChanged(ESlateCheckBoxState::Type CheckState, FCategoryViewModel Category)
+{
+	if (CheckState == ESlateCheckBoxState::Checked)
+	{
+		ViewModel->SetSelectedCategory(Category);
+	}
+}
+
+ESlateCheckBoxState::Type SAddContentWidget::GetCategoryCheckBoxCheckState(FCategoryViewModel Category) const
+{
+	return (Category == ViewModel->GetSelectedCategory()) ? ESlateCheckBoxState::Checked : ESlateCheckBoxState::Unchecked;
+}
+
+void SAddContentWidget::SearchTextChanged(const FText& SearchText)
+{
+	ViewModel->SetSearchText(SearchText);
+}
+
+void SAddContentWidget::ContentSourceTileViewSelectionChanged(TSharedPtr<FContentSourceViewModel> SelectedContentSource, ESelectInfo::Type SelectInfo)
+{
+	ViewModel->SetSelectedContentSource(SelectedContentSource);
+}
+
+FReply SAddContentWidget::ContentSourceTileDragDetected(const FGeometry& MyGeometry, const FPointerEvent& MouseEvent)
+{
+	TSharedPtr<FContentSourceViewModel> SelectedContentSource = ViewModel->GetSelectedContentSource();
+	if (SelectedContentSource.IsValid())
+	{
+		return FReply::Handled().BeginDragDrop(FContentSourceDragDropOp::CreateShared(SelectedContentSource));
+	}
+	return FReply::Unhandled();
+}
+
+FReply SAddContentWidget::PreviousScreenshotButtonClicked(TSharedPtr<FContentSourceViewModel> ContentSource)
+{
+	ContentSource->SelectPreviousScreenshotBrush();
+	return FReply::Handled();
+}
+
+FReply SAddContentWidget::NextScreenshotButtonClicked(TSharedPtr<FContentSourceViewModel> ContentSource)
+{
+	ContentSource->SelectNextScreenshotBrush();
+	return FReply::Handled();
+}
+
+FReply SAddContentWidget::RemoveAddedContentSourceClicked(TSharedPtr<FContentSourceViewModel> ContentSource)
+{
+	ViewModel->RemoveFromAddList(ContentSource);
+	return FReply::Handled();
+}
+
+void SAddContentWidget::CategoriesChanged()
+{
+	CategoryTabsContainer->SetContent(CreateCategoryTabs());
+}
+
+void SAddContentWidget::ContentSourcesChanged()
+{
+	ContentSourceTileView->RequestListRefresh();
+}
+
+void SAddContentWidget::SelectedContentSourceChanged()
+{
+	ContentSourceTileView->SetSelection(ViewModel->GetSelectedContentSource(), ESelectInfo::Direct);
+	ContentSourceDetailContainer->SetContent(CreateContentSourceDetail(ViewModel->GetSelectedContentSource()));
+}
+
+void SAddContentWidget::AddListChanged()
+{
+	AddListView->RequestListRefresh();
+	ContentSourcesToAdd.Empty();
+	for (const TSharedPtr<FContentSourceViewModel>& ContentSource : *ViewModel->GetAddList())
+	{
+		ContentSourcesToAdd.Add(ContentSource->GetContentSource());
+	}
+}
+
+SAddContentWidget::~SAddContentWidget()
+{
+}
\ No newline at end of file
diff --git a/Engine/Source/Editor/AddContentDialog/Private/SAddContentWidget.h b/Engine/Source/Editor/AddContentDialog/Private/SAddContentWidget.h
new file mode 100644
--- /dev/null
+++ b/Engine/Source/Editor/AddContentDialog/Private/SAddContentWidget.h
@@ -0,0 +1,107 @@
+// Copyright 1998-2014 Epic Games, Inc. All Rights Reserved.
+
+#pragma once
+
+class FAddContentWidgetViewModel;
+class FContentSourceViewModel;
+class FCategoryViewModel;
+
+/** A widget which allows the user to select multiple options from content which is available to be added to the project. */
+class SAddContentWidget : public SCompoundWidget
+{
+public:
+	DECLARE_DELEGATE(FOnAddListChanged);
+
+	SLATE_BEGIN_ARGS(SAddContentWidget)
+	{}
+
+	SLATE_END_ARGS()
+
+	~SAddContentWidget();
+
+	void Construct(const FArguments& InArgs);
+
+	/** Gets the content sources which have been selected by the user for addition to the project */
+	const TArray<TSharedPtr<IContentSource>>* GetContentSourcesToAdd();
+
+private:
+	/** Creates a strip of tabs which display and allow selecting categories. */
+	TSharedRef<SWidget> CreateCategoryTabs();
+
+	/** Creates a tile view which displays the content sources in the selected category. */
+	TSharedRef<SWidget> CreateContentSourceTileView();
+	
+	/** Creates the widget which represents each content source in the content source tile view. */
+	TSharedRef<ITableRow> CreateContentSourceIconTile(TSharedPtr<FContentSourceViewModel> ContentSource, const TSharedRef<STableViewBase>& OwnerTable);
+
+	/** Creates a widget representing detailed information about a single content source. */
+	TSharedRef<SWidget> CreateContentSourceDetail(TSharedPtr<FContentSourceViewModel> ContentSource);
+
+	/** Creates a widget carousel for displaying the set of screenshots for a specific content source. */
+	TSharedRef<SWidget> CreateScreenshotCarousel(TSharedPtr<FContentSourceViewModel> ContentSource);
+
+	/** Creates a list view represnting content which the user has selected for addition to the project. */
+	TSharedRef<SWidget> CreateAddListView();
+	
+	/** Creates a row for the list view represnting content selected for addition to the project */
+	TSharedRef<ITableRow> CreateAddContentSourceRow(TSharedPtr<FContentSourceViewModel> ContentSource, const TSharedRef<STableViewBase>& OwnerTable);
+
+
+
+	/** Handles the user clicking on one of the check boxes representing the category tabs. */
+	void CategoryCheckBoxCheckStateChanged(ESlateCheckBoxState::Type CheckState, FCategoryViewModel Category);
+
+	/** Gets the check state for one of the check boxes representing the category tabs. */
+	ESlateCheckBoxState::Type GetCategoryCheckBoxCheckState(FCategoryViewModel Category) const;
+
+	/** Handles the text in the search box changing. */
+	void SearchTextChanged(const FText& SearchText);
+
+	/** Handles the selection in the content source tile view changing */
+	void ContentSourceTileViewSelectionChanged(TSharedPtr<FContentSourceViewModel> SelectedContentSource, ESelectInfo::Type SelectInfo);
+
+	/** Handles an item from the content source tile view being dragged. */
+	FReply ContentSourceTileDragDetected(const FGeometry& MyGeometry, const FPointerEvent& MouseEvent);
+
+	/** Handles the previous screenshot button in the screenshot carousel being pressed. */
+	FReply PreviousScreenshotButtonClicked(TSharedPtr<FContentSourceViewModel> ContentSource);
+
+	/** Handles the next screenshot button in the screenshot carousel being pressed */
+	FReply NextScreenshotButtonClicked(TSharedPtr<FContentSourceViewModel> ContentSource);
+
+	/** Handles the remove button being pressed on a row in the list of content selected for addition to the project. */
+	FReply RemoveAddedContentSourceClicked(TSharedPtr<FContentSourceViewModel> ContentSource);
+
+
+
+	/** Handles the available categories changing on the view model. */
+	void CategoriesChanged();
+
+	/** Handles the available content sources changing on the view model. */
+	void ContentSourcesChanged();
+
+	/** Handles the selected content source changing on the view model. */
+	void SelectedContentSourceChanged();
+
+	/** Handles the list of items selected for addition to the project changing on the view model. */
+	void AddListChanged();
+
+private:
+	/** The view model which represents the current data of the UI. */
+	TSharedPtr<FAddContentWidgetViewModel> ViewModel;
+
+	/** The tile view which displays the content sources in the currently selected category. */
+	TSharedPtr<STileView<TSharedPtr<FContentSourceViewModel>>> ContentSourceTileView;
+
+	/** The placeholder widget which holds the category tab strip. */
+	TSharedPtr<SBox> CategoryTabsContainer;
+
+	/** The placeholder widget which holds the detail view for the currently selected content source. */
+	TSharedPtr<SBox> ContentSourceDetailContainer;
+
+	/** The list view which displays the content sources currently selected for addition to the project. */
+	TSharedPtr<SListView<TSharedPtr<FContentSourceViewModel>>> AddListView;
+
+	/** The content sources which the user has selected for addition to the project. */
+	TArray<TSharedPtr<IContentSource>> ContentSourcesToAdd;
+};
\ No newline at end of file
diff --git a/Engine/Source/Editor/AddContentDialog/Private/ViewModels/AddContentWidgetViewModel.cpp b/Engine/Source/Editor/AddContentDialog/Private/ViewModels/AddContentWidgetViewModel.cpp
new file mode 100644
--- /dev/null
+++ b/Engine/Source/Editor/AddContentDialog/Private/ViewModels/AddContentWidgetViewModel.cpp
@@ -0,0 +1,183 @@
+// Copyright 1998-2014 Epic Games, Inc. All Rights Reserved.
+
+#include "AddContentDialogPCH.h"
+
+#include "ImageWrapper.h"
+#include "ModuleManager.h"
+
+FAddContentWidgetViewModel::FAddContentWidgetViewModel()
+{
+}
+
+TSharedRef<FAddContentWidgetViewModel> FAddContentWidgetViewModel::CreateShared()
+{
+	TSharedPtr<FAddContentWidgetViewModel> Shared = MakeShareable(new FAddContentWidgetViewModel());
+	Shared->Initialize();
+	return Shared.ToSharedRef();
+}
+
+void FAddContentWidgetViewModel::Initialize()
+{
+	IAddContentDialogModule& AddContentDialogModule = FModuleManager::LoadModuleChecked<IAddContentDialogModule>("AddContentDialog");
+
+	for (auto ContentSourceProvider : *AddContentDialogModule.GetContentSourceProviderManager()->GetContentSourceProviders())
+	{
+		ContentSourceProviders.Add(ContentSourceProvider);
+		ContentSourceProvider->SetContentSourcesChanged(FOnContentSourcesChanged::CreateSP(this, &FAddContentWidgetViewModel::ContentSourcesChanged));
+	}
+
+	ContentSourceFilter = TSharedPtr<ContentSourceTextFilter>(new ContentSourceTextFilter(
+		ContentSourceTextFilter::FItemToStringArray::CreateSP(this, &FAddContentWidgetViewModel::TransformContentSourceToStrings)));
+
+	BuildContentSourceViewModels();
+}
+
+const TArray<FCategoryViewModel>* FAddContentWidgetViewModel::GetCategories()
+{
+	return &Categories;
+}
+
+void FAddContentWidgetViewModel::SetOnCategoriesChanged(FOnCategoriesChanged OnCategoriesChangedIn)
+{
+	OnCategoriesChanged = OnCategoriesChangedIn;
+}
+
+FCategoryViewModel FAddContentWidgetViewModel::GetSelectedCategory()
+{
+	return SelectedCategory;
+}
+
+void FAddContentWidgetViewModel::SetSelectedCategory(FCategoryViewModel SelectedCategoryIn)
+{
+	SelectedCategory = SelectedCategoryIn;
+	FilterContentSourceViewModels();
+	OnSelectedContentSourceChanged.ExecuteIfBound();
+}
+
+void FAddContentWidgetViewModel::SetSearchText(FText SearchTextIn)
+{
+	SearchText = SearchTextIn;
+	ContentSourceFilter->SetRawFilterText(SearchTextIn);
+	FilterContentSourceViewModels();
+}
+
+const TArray<TSharedPtr<FContentSourceViewModel>>* FAddContentWidgetViewModel::GetContentSources()
+{
+	return &FilteredContentSourceViewModels;
+}
+
+void FAddContentWidgetViewModel::SetOnContentSourcesChanged(FOnContentSourcesChanged OnContentSourcesChangedIn)
+{
+	OnContentSourcesChanged = OnContentSourcesChangedIn;
+}
+
+TSharedPtr<FContentSourceViewModel> FAddContentWidgetViewModel::GetSelectedContentSource()
+{
+	TSharedPtr<FContentSourceViewModel>* SelectedContentSource = CategoryToSelectedContentSourceMap.Find(SelectedCategory);
+	if (SelectedContentSource != nullptr)
+	{
+		return *SelectedContentSource;
+	}
+	return TSharedPtr<FContentSourceViewModel>();
+}
+
+void FAddContentWidgetViewModel::SetSelectedContentSource(TSharedPtr<FContentSourceViewModel> SelectedContentSourceIn)
+{
+	// Ignore null selections.
+	if (SelectedContentSourceIn.IsValid())
+	{
+		// Ignore selecting the currently selected item.
+		TSharedPtr<FContentSourceViewModel> SelectedContentSource = GetSelectedContentSource();
+		if (SelectedContentSource != SelectedContentSourceIn)
+		{
+			CategoryToSelectedContentSourceMap.Add(SelectedCategory) = SelectedContentSourceIn;
+			OnSelectedContentSourceChanged.ExecuteIfBound();
+		}
+	}
+}
+
+void FAddContentWidgetViewModel::SetOnSelectedContentSourceChanged(FOnSelectedContentSourceChanged OnSelectedContentSourceChangedIn)
+{
+	OnSelectedContentSourceChanged = OnSelectedContentSourceChangedIn;
+}
+
+void FAddContentWidgetViewModel::BuildContentSourceViewModels()
+{
+	Categories.Empty();
+	ContentSourceViewModels.Empty();
+	FilteredContentSourceViewModels.Empty();
+	
+	for (auto ContentSourceProvider : ContentSourceProviders)
+	{
+		for (auto ContentSource : ContentSourceProvider->GetContentSources())
+		{
+			TSharedPtr<FContentSourceViewModel> ContentSourceViewModel = MakeShareable(new FContentSourceViewModel(ContentSource));
+			if (Categories.Contains(ContentSourceViewModel->GetCategory()) == false)
+			{
+				Categories.Add(ContentSourceViewModel->GetCategory());
+			}
+			ContentSourceViewModels.Add(ContentSourceViewModel);
+		}
+	}
+
+	if (Categories.Num() > 0)
+	{
+		SelectedCategory = Categories[0];
+	}
+	OnCategoriesChanged.ExecuteIfBound();
+	FilterContentSourceViewModels();
+}
+
+void FAddContentWidgetViewModel::FilterContentSourceViewModels()
+{
+	FilteredContentSourceViewModels.Empty();
+	for (auto ContentSource : ContentSourceViewModels)
+	{
+		if ((ContentSource->GetCategory() == SelectedCategory) && (ContentSourceFilter.IsValid() == false || ContentSourceFilter->PassesFilter(ContentSource)))
+		{
+			FilteredContentSourceViewModels.Add(ContentSource);
+		}
+	}
+	OnContentSourcesChanged.ExecuteIfBound();
+
+	if (FilteredContentSourceViewModels.Contains(GetSelectedContentSource()) == false)
+	{
+		TSharedPtr<FContentSourceViewModel> NewSelectedContentSource = FilteredContentSourceViewModels.Num() > 0 ? FilteredContentSourceViewModels[0] : TSharedPtr<FContentSourceViewModel>();
+		SetSelectedContentSource(NewSelectedContentSource);
+	}
+}
+
+void FAddContentWidgetViewModel::TransformContentSourceToStrings(TSharedPtr<FContentSourceViewModel> Item, OUT TArray<FString>& Array)
+{
+	Array.Add(Item->GetName().ToString());
+}
+
+void FAddContentWidgetViewModel::ContentSourcesChanged()
+{
+	BuildContentSourceViewModels();
+}
+
+void FAddContentWidgetViewModel::AppendToAddList(TSharedPtr<FContentSourceViewModel> ContentSourceToAppend)
+{
+	if (AddList.Contains(ContentSourceToAppend) == false)
+	{
+		AddList.Add(ContentSourceToAppend);
+		OnAddListChanged.ExecuteIfBound();
+	}
+}
+
+void FAddContentWidgetViewModel::RemoveFromAddList(TSharedPtr<FContentSourceViewModel> ContentSourceToRemove)
+{
+	AddList.Remove(ContentSourceToRemove);
+	OnAddListChanged.ExecuteIfBound();
+}
+
+const TArray<TSharedPtr<FContentSourceViewModel>>* FAddContentWidgetViewModel::GetAddList()
+{
+	return &AddList;
+}
+
+void FAddContentWidgetViewModel::SetOnAddListChanged(FOnAddListChanged OnAddListChangedIn)
+{
+	OnAddListChanged = OnAddListChangedIn;
+}
\ No newline at end of file
diff --git a/Engine/Source/Editor/AddContentDialog/Private/ViewModels/AddContentWidgetViewModel.h b/Engine/Source/Editor/AddContentDialog/Private/ViewModels/AddContentWidgetViewModel.h
new file mode 100644
--- /dev/null
+++ b/Engine/Source/Editor/AddContentDialog/Private/ViewModels/AddContentWidgetViewModel.h
@@ -0,0 +1,121 @@
+// Copyright 1998-2014 Epic Games, Inc. All Rights Reserved.
+
+#pragma once
+
+/** The view model for the SAddContentWidget. */
+class FAddContentWidgetViewModel : public TSharedFromThis<FAddContentWidgetViewModel>
+{
+public:
+	DECLARE_DELEGATE(FOnCategoriesChanged)
+	DECLARE_DELEGATE(FOnContentSourcesChanged);
+	DECLARE_DELEGATE(FOnSelectedContentSourceChanged);
+	DECLARE_DELEGATE(FOnAddListChanged);
+
+	typedef TTextFilter<TSharedPtr<FContentSourceViewModel>> ContentSourceTextFilter;
+
+	/** Creates a shared reference to a new view model. */
+	static TSharedRef<FAddContentWidgetViewModel> CreateShared();
+
+	/** Gets the view models for the current set of content source categories. */
+	const TArray<FCategoryViewModel>* GetCategories();
+
+	/** Sets the delegate which should be executed when the set of categories changes. */
+	void SetOnCategoriesChanged(FOnCategoriesChanged OnCategoriesChangedIn);
+
+	/** Gets the view model for the currently selected category.  Only content sources which match
+		the selected category will be returned by GetContentSources(). */
+	FCategoryViewModel GetSelectedCategory();
+
+	/** Sets the currently selected category view model. Only content sources which match the selected
+		category will be returned by GetContentSources(). */
+	void SetSelectedCategory(FCategoryViewModel SelectedCategoryIn);
+
+	/** Sets search text which should be used to filter the content sources. */
+	void SetSearchText(FText SearchTextIn);
+
+	/** Gets a filtered array of content sources which match both the selected category and the search
+		text if it has been set. */
+	const TArray<TSharedPtr<FContentSourceViewModel>>* GetContentSources();
+
+	/** Sets the delegate which should be executed when the current set of content sources returned by
+		GetContentSources changes. */
+	void SetOnContentSourcesChanged(FOnContentSourcesChanged OnContentSourcesChangedIn);
+
+	/** Gets the currently selected content source. */
+	TSharedPtr<FContentSourceViewModel> GetSelectedContentSource();
+
+	/** Sets the currently selected content source. */
+	void SetSelectedContentSource(TSharedPtr<FContentSourceViewModel> SelectedContentSourceIn);
+
+	/** Sets the delegate which should be executed when the selected content source changes. */
+	void SetOnSelectedContentSourceChanged(FOnSelectedContentSourceChanged OnSelectedContentSourceChangedIn);
+
+	/** Gets the array of content sources which the user has selected for addition to the project. */
+	const TArray<TSharedPtr<FContentSourceViewModel>>* GetAddList();
+
+	/** Adds a content source to the array of content sources which the user has selected for addition to the project. */
+	void AppendToAddList(TSharedPtr<FContentSourceViewModel> ContentSourceToAppend);
+
+	/** Removes a content source from the array of content sources which the user has selected for addition to the project. */
+	void RemoveFromAddList(TSharedPtr<FContentSourceViewModel> ContentSourceToRemove);
+
+	/** Sets the delegate which should be executed when the contents of the array returned by GetAddList changes. */
+	void SetOnAddListChanged(FOnAddListChanged OnAddListChangedIn);
+
+private:
+	FAddContentWidgetViewModel();
+	void Initialize();
+
+	/** Builds view models for all available content sources. */
+	void BuildContentSourceViewModels();
+
+	/** Filters the current set of content sources based on the selected category and the search text. */
+	void FilterContentSourceViewModels();
+
+	/** Converts a content source item to an array of strings for processing by the TTextFilter. */
+	void TransformContentSourceToStrings(TSharedPtr<FContentSourceViewModel> Item, OUT TArray< FString >& Array);
+
+	/** Handles notification from the content source providers when their content source arrays change. */
+	void ContentSourcesChanged();
+
+private:
+	/** The array of content source providers who's content sources are being displayed. */
+	TArray<TSharedPtr<IContentSourceProvider>> ContentSourceProviders;
+
+	/** The view models for the available categories. */
+	TArray<FCategoryViewModel> Categories;
+
+	/** A combined array of all content sources from all content source providers. */
+	TArray<TSharedPtr<FContentSourceViewModel>> ContentSourceViewModels;
+
+	/** A filtered array of content sources based on the currently selected category and the search text. */
+	TArray<TSharedPtr<FContentSourceViewModel>> FilteredContentSourceViewModels;
+
+	/** An array of the content sources which the user has selected for addition to the project. */
+	TArray<TSharedPtr<FContentSourceViewModel>> AddList;
+
+	/** A map which keeps track of the currently selected content source for each category. */
+	TMap<FCategoryViewModel, TSharedPtr<FContentSourceViewModel>> CategoryToSelectedContentSourceMap;
+
+	/** The view model for the currently selected category. */
+	FCategoryViewModel SelectedCategory;
+
+	/** The current search text. */
+	FText SearchText;
+
+	/** The delegate which is executed when the available categories change. */
+	FOnCategoriesChanged OnCategoriesChanged;
+
+	/** The delegate which is executed when the filtered content sources change. */
+	FOnContentSourcesChanged OnContentSourcesChanged;
+
+	/** The delegate which is executed when the currently selected content source changes. */
+	FOnSelectedContentSourceChanged OnSelectedContentSourceChanged;
+
+	/** The delegate which is executed when the contents of the array of content souces to add
+		changes. */
+	FOnAddListChanged OnAddListChanged;
+
+	/** The filter which is used to filter the content sources based on the search text. */
+	TSharedPtr<ContentSourceTextFilter> ContentSourceFilter;
+};
\ No newline at end of file
diff --git a/Engine/Source/Editor/AddContentDialog/Private/ViewModels/CategoryViewModel.cpp b/Engine/Source/Editor/AddContentDialog/Private/ViewModels/CategoryViewModel.cpp
new file mode 100644
--- /dev/null
+++ b/Engine/Source/Editor/AddContentDialog/Private/ViewModels/CategoryViewModel.cpp
@@ -0,0 +1,56 @@
+// Copyright 1998-2014 Epic Games, Inc. All Rights Reserved.
+
+#include "AddContentDialogPCH.h"
+
+#define LOCTEXT_NAMESPACE "ContentSourceViewModel"
+
+FCategoryViewModel::FCategoryViewModel()
+{
+	Category = EContentSourceCategory::Unknown;
+	Initialize();
+}
+
+FCategoryViewModel::FCategoryViewModel(EContentSourceCategory InCategory)
+{
+	Category = InCategory;
+	Initialize();
+}
+
+FText FCategoryViewModel::GetText() const
+{
+	return Text;
+}
+
+const FSlateBrush* FCategoryViewModel::GetIconBrush() const
+{
+	return IconBrush;
+}
+
+inline bool FCategoryViewModel::operator==(const FCategoryViewModel& Other) const
+{
+	return Other.Text.EqualTo(Text) && (Other.IconBrush == IconBrush);
+}
+
+uint32 FCategoryViewModel::GetTypeHash() const
+{
+	return (uint32)Category;
+}
+
+void FCategoryViewModel::Initialize()
+{
+	switch (Category)
+	{
+	case EContentSourceCategory::BlueprintFeature:
+		Text = LOCTEXT("BlueprintFeature", "Blueprint Feature");
+		IconBrush = FAddContentDialogStyle::Get().GetBrush("AddContentDialog.BlueprintFeatureCategory");
+		break;
+	case EContentSourceCategory::CodeFeature:
+		Text = LOCTEXT("CodeFeature", "C++ Feature");
+		IconBrush = FAddContentDialogStyle::Get().GetBrush("AddContentDialog.CodeFeatureCategory");
+		break;
+	default:
+		Text = LOCTEXT("Miscellaneous", "Miscellaneous");
+		IconBrush = FAddContentDialogStyle::Get().GetBrush("AddContentDialog.DefaultCategory");
+		break;
+	}
+}
\ No newline at end of file
diff --git a/Engine/Source/Editor/AddContentDialog/Private/ViewModels/CategoryViewModel.h b/Engine/Source/Editor/AddContentDialog/Private/ViewModels/CategoryViewModel.h
new file mode 100644
--- /dev/null
+++ b/Engine/Source/Editor/AddContentDialog/Private/ViewModels/CategoryViewModel.h
@@ -0,0 +1,35 @@
+// Copyright 1998-2014 Epic Games, Inc. All Rights Reserved.
+
+#pragma once
+
+/** A view model for displaying a content source category in the UI. */
+class FCategoryViewModel
+{
+public:
+	FCategoryViewModel();
+
+	FCategoryViewModel(EContentSourceCategory InCategory);
+
+	/** Gets the display name of the category. */
+	FText GetText() const;
+
+	/** Gets the brush which should be used to draw the icon for the category. */
+	const FSlateBrush* GetIconBrush() const;
+
+	inline bool operator==(const FCategoryViewModel& Other) const;
+
+	uint32 GetTypeHash() const;
+
+private:
+	void Initialize();
+
+private:
+	EContentSourceCategory Category;
+	FText Text;
+	const FSlateBrush* IconBrush;
+};
+
+inline uint32 GetTypeHash(const FCategoryViewModel& CategoryViewModel)
+{
+	return CategoryViewModel.GetTypeHash();
+}
\ No newline at end of file
diff --git a/Engine/Source/Editor/AddContentDialog/Private/ViewModels/ContentSourceViewModel.cpp b/Engine/Source/Editor/AddContentDialog/Private/ViewModels/ContentSourceViewModel.cpp
new file mode 100644
--- /dev/null
+++ b/Engine/Source/Editor/AddContentDialog/Private/ViewModels/ContentSourceViewModel.cpp
@@ -0,0 +1,167 @@
+// Copyright 1998-2014 Epic Games, Inc. All Rights Reserved.
+
+#include "AddContentDialogPCH.h"
+
+#include "ImageWrapper.h"
+#include "ModuleManager.h"
+
+#define LOCTEXT_NAMESPACE "ContentSourceViewModel"
+
+const FString DefaultLanguageCode = "en";
+
+FContentSourceViewModel::FContentSourceViewModel(TSharedPtr<IContentSource> ContentSourceIn)
+{
+	ContentSource = ContentSourceIn;
+	SetupBrushes();
+	Category = FCategoryViewModel(ContentSource->GetCategory());
+}
+
+TSharedPtr<IContentSource> FContentSourceViewModel::GetContentSource()
+{
+	return ContentSource;
+}
+
+FText FContentSourceViewModel::GetName()
+{
+	FString CurrentLanguage = FInternationalization::Get().GetCurrentCulture()->GetTwoLetterISOLanguageName();
+	if (NameText.GetTwoLetterLanguage() != CurrentLanguage)
+	{
+		NameText = ChooseLocalizedText(ContentSource->GetLocalizedNames(), CurrentLanguage);
+	}
+	return NameText.GetText();
+}
+
+FText FContentSourceViewModel::GetDescription()
+{
+	FString CurrentLanguage = FInternationalization::Get().GetCurrentCulture()->GetTwoLetterISOLanguageName();
+	if (DescriptionText.GetTwoLetterLanguage() != CurrentLanguage)
+	{
+		DescriptionText = ChooseLocalizedText(ContentSource->GetLocalizedDescriptions(), CurrentLanguage);
+	}
+	return DescriptionText.GetText();
+}
+
+FCategoryViewModel FContentSourceViewModel::GetCategory()
+{
+	return Category;
+}
+
+TSharedPtr<FSlateBrush> FContentSourceViewModel::GetIconBrush()
+{
+	return IconBrush;
+}
+
+TArray<TSharedPtr<FSlateBrush>>* FContentSourceViewModel::GetScreenshotBrushes()
+{
+	return &ScreenshotBrushes;
+}
+
+
+const FSlateBrush* FContentSourceViewModel::GetSelectedScreenshotBrush() const
+{
+	return SelectedScreenshotBrush.Get();
+}
+
+void FContentSourceViewModel::SelectNextScreenshotBrush()
+{
+	if (ScreenshotBrushes.Num() > 1)
+	{
+		int currentIndex = ScreenshotBrushes.IndexOfByKey(SelectedScreenshotBrush);
+		int newIndex = currentIndex + 1;
+		if (newIndex < 0)
+		{
+			newIndex += ScreenshotBrushes.Num();
+		}
+		else if (newIndex >= ScreenshotBrushes.Num())
+		{
+			newIndex -= ScreenshotBrushes.Num();
+		}
+		SelectedScreenshotBrush = ScreenshotBrushes[newIndex];
+	}
+}
+
+void FContentSourceViewModel::SelectPreviousScreenshotBrush()
+{
+	if (ScreenshotBrushes.Num() > 1)
+	{
+		int currentIndex = ScreenshotBrushes.IndexOfByKey(SelectedScreenshotBrush);
+		int newIndex = currentIndex - 1;
+		if (newIndex < 0)
+		{
+			newIndex += ScreenshotBrushes.Num();
+		}
+		else if (newIndex >= ScreenshotBrushes.Num())
+		{
+			newIndex -= ScreenshotBrushes.Num();
+		}
+		SelectedScreenshotBrush = ScreenshotBrushes[newIndex];
+	}
+}
+
+void FContentSourceViewModel::SetupBrushes()
+{
+	FString IconBrushName = GetName().ToString() + "_" + ContentSource->GetIconData()->GetName();
+	IconBrush = CreateBrushFromRawData(FName(*IconBrushName), *ContentSource->GetIconData()->GetData());
+
+	for (TSharedPtr<FImageData> ScreenshotData : ContentSource->GetScreenshotData())
+	{
+		FString ScreenshotBrushName = GetName().ToString() + "_" + ScreenshotData->GetName();
+		ScreenshotBrushes.Add(CreateBrushFromRawData(FName(*ScreenshotBrushName), *ScreenshotData->GetData()));
+	}
+	if (ScreenshotBrushes.Num())
+	{
+		SelectedScreenshotBrush = ScreenshotBrushes[0];
+	}
+}
+
+TSharedPtr<FSlateDynamicImageBrush> FContentSourceViewModel::CreateBrushFromRawData(FName ResourceName, const TArray< uint8 >& RawData) const
+{
+	TSharedPtr< FSlateDynamicImageBrush > Brush;
+
+	uint32 BytesPerPixel = 4;
+	int32 Width = 0;
+	int32 Height = 0;
+
+	bool bSucceeded = false;
+	TArray<uint8> DecodedImage;
+	IImageWrapperModule& ImageWrapperModule = FModuleManager::LoadModuleChecked<IImageWrapperModule>(FName("ImageWrapper"));
+	IImageWrapperPtr ImageWrapper = ImageWrapperModule.CreateImageWrapper(EImageFormat::PNG);
+	if (ImageWrapper.IsValid() && ImageWrapper->SetCompressed(RawData.GetData(), RawData.Num()))
+	{
+		Width = ImageWrapper->GetWidth();
+		Height = ImageWrapper->GetHeight();
+
+		const TArray<uint8>* RawImageData = NULL;
+		if (ImageWrapper->GetRaw(ERGBFormat::BGRA, 8, RawImageData))
+		{
+			DecodedImage.AddUninitialized(Width * Height * BytesPerPixel);
+			DecodedImage = *RawImageData;
+			bSucceeded = true;
+		}
+	}
+
+	if (bSucceeded)
+	{
+		Brush = FSlateDynamicImageBrush::CreateWithImageData(ResourceName, FVector2D(ImageWrapper->GetWidth(), ImageWrapper->GetHeight()), DecodedImage);
+	}
+
+	return Brush;
+}
+
+FLocalizedText FContentSourceViewModel::ChooseLocalizedText(TArray<FLocalizedText> Choices, FString LanguageCode)
+{
+	FLocalizedText Default;
+	for (const FLocalizedText& Choice : Choices)
+	{
+		if (Choice.GetTwoLetterLanguage() == LanguageCode)
+		{
+			return Choice;
+			break;
+		}
+		else if (Choice.GetTwoLetterLanguage() == DefaultLanguageCode)
+		{
+			Default = Choice;
+		}
+	}
+	return Default;
+}
diff --git a/Engine/Source/Editor/AddContentDialog/Private/ViewModels/ContentSourceViewModel.h b/Engine/Source/Editor/AddContentDialog/Private/ViewModels/ContentSourceViewModel.h
new file mode 100644
--- /dev/null
+++ b/Engine/Source/Editor/AddContentDialog/Private/ViewModels/ContentSourceViewModel.h
@@ -0,0 +1,74 @@
+// Copyright 1998 - 2014 Epic Games, Inc.All Rights Reserved.
+
+#pragma once
+
+/** A view model for displaying and interacting with an IContentSource in the FAddContentDialog. */
+class FContentSourceViewModel : public TSharedFromThis<FContentSourceViewModel>
+{
+public:
+	/** Creates a view model for a supplied content source. */
+	FContentSourceViewModel(TSharedPtr<IContentSource> ContentSourceIn);
+
+	/** Gets the content source represented by this view model. */
+	TSharedPtr<IContentSource> GetContentSource();
+
+	/** Gets the display name for this content source. */
+	FText GetName();
+
+	/** Gets the description of this content source. */
+	FText GetDescription();
+
+	/** Gets the view model for the category for this content source. */
+	FCategoryViewModel GetCategory();
+
+	/** Gets the brush which should be used to draw the icon representation of this content source. */
+	TSharedPtr<FSlateBrush> GetIconBrush();
+
+	/** Gets an array or brushes which should be used to display screenshots for this content source. */
+	TArray<TSharedPtr<FSlateBrush>>* GetScreenshotBrushes();
+
+	/** Gets the brush for the currently selected screenshot. */
+	const FSlateBrush* GetSelectedScreenshotBrush() const;
+
+	/** Selects the next available screenshot.  When at the last screenshot this will start again at
+		the beginning of the array. */
+	void SelectNextScreenshotBrush();
+
+	/** Selects the previous available screenshot.  When at the first screenshot this will start again at
+		the end of the array. */
+	void SelectPreviousScreenshotBrush();
+
+private:
+	/** Sets up brushes from the images data supplied by the IContentSource. */
+	void SetupBrushes();
+
+	/** Creates a slate brush from raw binary PNG formatted image data and the supplied FName. */
+	TSharedPtr<FSlateDynamicImageBrush> CreateBrushFromRawData(FName ResourceName, const TArray<uint8>& RawData) const;
+
+	/** Selects an FLocalizedText from an array which matches either the supplied language code, or the default language code. */
+	FLocalizedText ChooseLocalizedText(TArray<FLocalizedText> Choices, FString LanguageCode);
+
+private:
+	/** The content source represented by this view model. */
+	TSharedPtr<IContentSource> ContentSource;
+
+	/** The brush which should be used to draw the icon representation of this content source. */
+	TSharedPtr<FSlateBrush> IconBrush;
+
+	/** An array or brushes which should be used to display screenshots for this content source. */
+	TArray<TSharedPtr<FSlateBrush>> ScreenshotBrushes;
+
+	/** The brush for the currently selected screenshot. */
+	TSharedPtr<FSlateBrush> SelectedScreenshotBrush;
+
+	/** The view model for the category for this content source. */
+	FCategoryViewModel Category;
+
+	/** The FLocalizedText representing the name of the content source, in the language which was active the
+		last time it was requested, or the default language if a translation was not available. */
+	FLocalizedText NameText;
+
+	/** The FLocalizedText representing the description of the content source, in the language which was active
+		the last time it was requested, or the default language if a translation was not available. */
+	FLocalizedText DescriptionText;
+};
\ No newline at end of file
diff --git a/Engine/Source/Editor/AddContentDialog/Public/IAddContentDialogModule.h b/Engine/Source/Editor/AddContentDialog/Public/IAddContentDialogModule.h
new file mode 100644
--- /dev/null
+++ b/Engine/Source/Editor/AddContentDialog/Public/IAddContentDialogModule.h
@@ -0,0 +1,19 @@
+// Copyright 1998-2014 Epic Games, Inc. All Rights Reserved.
+
+#pragma once
+
+#include "ModuleInterface.h"
+
+class FContentSourceProviderManager;
+
+/** Defines methods for interacting with the Add Content Dialog.  */
+class IAddContentDialogModule : public IModuleInterface
+{
+
+public:
+	/** Gets the object responsible for managing content source providers. */
+	virtual TSharedRef<FContentSourceProviderManager> GetContentSourceProviderManager() = 0;
+
+	/** Creates a dialog for adding existing content to a project. */
+	virtual TSharedRef<SWindow> CreateDialogWindow() = 0;
+};
\ No newline at end of file
diff --git a/Engine/Source/Editor/ContentBrowser/ContentBrowser.Build.cs b/Engine/Source/Editor/ContentBrowser/ContentBrowser.Build.cs
--- a/Engine/Source/Editor/ContentBrowser/ContentBrowser.Build.cs
+++ b/Engine/Source/Editor/ContentBrowser/ContentBrowser.Build.cs
@@ -15,7 +15,7 @@ public class ContentBrowser : ModuleRules
                 "MainFrame",
 				"SourceControl",
 				"SourceControlWindows",
-                "ReferenceViewer"
+                "ReferenceViewer",
 			}
 		);
 
@@ -34,7 +34,8 @@ public class ContentBrowser : ModuleRules
 				"WorkspaceMenuStructure",
 				"UnrealEd",
 				"EditorWidgets",
-				"Projects"
+				"Projects",
+				"AddContentDialog"
 			}
 		);
 
diff --git a/Engine/Source/Editor/ContentBrowser/Private/SContentBrowser.cpp b/Engine/Source/Editor/ContentBrowser/Private/SContentBrowser.cpp
--- a/Engine/Source/Editor/ContentBrowser/Private/SContentBrowser.cpp
+++ b/Engine/Source/Editor/ContentBrowser/Private/SContentBrowser.cpp
@@ -11,6 +11,7 @@
 #include "AssetRegistryModule.h"
 #include "SDockTab.h"
 #include "GenericCommands.h"
+#include "IAddContentDialogModule.h"
 
 #define LOCTEXT_NAMESPACE "ContentBrowser"
 
@@ -173,6 +174,42 @@ void SContentBrowser::Construct( const FArguments& InArgs, const FName& InInstan
 							]
 						]
 
+						// Get content
+						+ SHorizontalBox::Slot()
+						.VAlign(VAlign_Center)
+						.HAlign(HAlign_Left)
+						[
+							SNew(SButton)
+							.ButtonStyle(FEditorStyle::Get(), "ToggleButton")
+							.ToolTipText(LOCTEXT("AddContentTooltip", "Get more content."))
+							.ContentPadding(0)
+							.Visibility(EVisibility::Collapsed) // TODO: Make the button visible once there are actually content packs available.
+							.OnClicked(this, &SContentBrowser::OnAddContentClicked)
+							[
+								SNew(SHorizontalBox)
+
+								// Get Content Icon
+								+ SHorizontalBox::Slot()
+								.AutoWidth()
+								.VAlign(VAlign_Center)
+								[
+									SNew(SImage)
+									.Image(FEditorStyle::GetBrush("ContentBrowser.AddContent"))
+								]
+
+								// Get Content Text
+								+ SHorizontalBox::Slot()
+								.AutoWidth()
+								.VAlign(VAlign_Center)
+								.Padding(0, 0, 2, 0)
+								[
+									SNew(STextBlock)
+									.TextStyle(FEditorStyle::Get(), "ContentBrowser.TopBar.Font")
+									.Text(LOCTEXT("GetContent", "Get Content"))
+								]
+							]
+						]
+
 						// Save
 						+ SHorizontalBox::Slot()
 						.FillWidth(1.0f)
@@ -1322,6 +1359,18 @@ FReply SContentBrowser::OnSaveClicked()
 	return FReply::Handled();
 }
 
+FReply SContentBrowser::OnAddContentClicked()
+{
+	IAddContentDialogModule& AddContentDialogModule = FModuleManager::LoadModuleChecked<IAddContentDialogModule>("AddContentDialog");
+	TSharedRef<SWindow> AddContentDialog = AddContentDialogModule.CreateDialogWindow();
+
+	FWidgetPath WidgetPath;
+	FSlateApplication::Get().GeneratePathToWidgetChecked(AsShared(), WidgetPath);
+	FSlateApplication::Get().AddWindowAsNativeChild(AddContentDialog, WidgetPath.GetWindow());
+
+	return FReply::Handled();
+}
+
 void SContentBrowser::OnAssetSelectionChanged(const FAssetData& SelectedAsset)
 {
 	if ( bIsPrimaryBrowser )
diff --git a/Engine/Source/Editor/ContentBrowser/Private/SContentBrowser.h b/Engine/Source/Editor/ContentBrowser/Private/SContentBrowser.h
--- a/Engine/Source/Editor/ContentBrowser/Private/SContentBrowser.h
+++ b/Engine/Source/Editor/ContentBrowser/Private/SContentBrowser.h
@@ -167,6 +167,8 @@ private:
 	/** Saves dirty content. */
 	FReply OnSaveClicked();
 
+	FReply OnAddContentClicked();
+
 	/** Handler for when the selection set in the asset view has changed. */
 	void OnAssetSelectionChanged(const FAssetData& SelectedAsset);
 
diff --git a/Engine/Source/Editor/EditorStyle/Private/SlateEditorStyle.cpp b/Engine/Source/Editor/EditorStyle/Private/SlateEditorStyle.cpp
--- a/Engine/Source/Editor/EditorStyle/Private/SlateEditorStyle.cpp
+++ b/Engine/Source/Editor/EditorStyle/Private/SlateEditorStyle.cpp
@@ -5054,6 +5054,7 @@ void FSlateEditorStyle::FStyle::SetupContentBrowserStyle()
 			.SetShadowColorAndOpacity( FLinearColor( 0, 0, 0, 0.9f ) ) );
 
 		Set( "ContentBrowser.SaveDirtyPackages", new IMAGE_BRUSH( "Icons/icon_file_saveall_40x", Icon25x25 ) );
+		Set( "ContentBrowser.AddContent", new IMAGE_BRUSH( "Icons/icon_AddContent_40x", Icon25x25 ) );
 		Set( "ContentBrowser.ImportPackage", new IMAGE_BRUSH( "Icons/icon_Import_40x", Icon25x25 ) );
 		Set( "ContentBrowser.ImportIcon", new IMAGE_BRUSH( "Icons/icon_Import_16x", Icon16x16 ) );
 
diff --git a/Engine/Source/Runtime/PakFile/Private/IPlatformFilePak.cpp b/Engine/Source/Runtime/PakFile/Private/IPlatformFilePak.cpp
--- a/Engine/Source/Runtime/PakFile/Private/IPlatformFilePak.cpp
+++ b/Engine/Source/Runtime/PakFile/Private/IPlatformFilePak.cpp
@@ -677,8 +677,8 @@ bool FPakPlatformFile::Initialize(IPlatformFile* Inner, const TCHAR* CmdLine)
 bool FPakPlatformFile::Mount(const TCHAR* InPakFilename, uint32 PakOrder, const TCHAR* InPath /*= NULL*/)
 {
 	bool bSuccess = false;
-	IFileHandle* PakHandle = LowerLevel->OpenRead(InPakFilename);
-	if (PakHandle != NULL)
+	TSharedPtr<IFileHandle> PakHandle = MakeShareable(LowerLevel->OpenRead(InPakFilename));
+	if (PakHandle.IsValid())
 	{
 		FPakFile* Pak = new FPakFile(LowerLevel, InPakFilename, bSigned);
 		if (Pak->IsValid())
