AddContentDialog - Add destructors and undefine loctext namespaces to fix the build.

diff --git a/Engine/Source/Editor/AddContentDialog/Private/AddContentDialogModule.cpp b/Engine/Source/Editor/AddContentDialog/Private/AddContentDialogModule.cpp
--- a/Engine/Source/Editor/AddContentDialog/Private/AddContentDialogModule.cpp
+++ b/Engine/Source/Editor/AddContentDialog/Private/AddContentDialogModule.cpp
@@ -37,8 +37,9 @@ public:
 
 private:
 	TSharedPtr<FContentSourceProviderManager> ContentSourceProviderManager;
-	FString TransferFilePath = FPaths::Combine(FPlatformProcess::UserSettingsDir(), TEXT("UnrealEngineLauncher"), TEXT("com"), TEXT("transfer.tmp"));
 
 };
 
-IMPLEMENT_MODULE(FAddContentDialogModule, AddContentDialog);
\ No newline at end of file
+IMPLEMENT_MODULE(FAddContentDialogModule, AddContentDialog);
+
+#undef LOCTEXT_NAMESPACE
\ No newline at end of file
diff --git a/Engine/Source/Editor/AddContentDialog/Private/ContentSourceProviders/FeaturePack/FeaturePackContentSource.cpp b/Engine/Source/Editor/AddContentDialog/Private/ContentSourceProviders/FeaturePack/FeaturePackContentSource.cpp
--- a/Engine/Source/Editor/AddContentDialog/Private/ContentSourceProviders/FeaturePack/FeaturePackContentSource.cpp
+++ b/Engine/Source/Editor/AddContentDialog/Private/ContentSourceProviders/FeaturePack/FeaturePackContentSource.cpp
@@ -110,4 +110,8 @@ void FFeaturePackContentSource::InstallToProject(FString InstallPath)
 	TArray<FString> AssetPaths;
 	AssetPaths.Add(FeaturePackPath);
 	AssetToolsModule.Get().ImportAssets(AssetPaths, InstallPath);
+}
+
+FFeaturePackContentSource::~FFeaturePackContentSource()
+{
 }
\ No newline at end of file
diff --git a/Engine/Source/Editor/AddContentDialog/Private/ContentSourceProviders/FeaturePack/FeaturePackContentSource.h b/Engine/Source/Editor/AddContentDialog/Private/ContentSourceProviders/FeaturePack/FeaturePackContentSource.h
--- a/Engine/Source/Editor/AddContentDialog/Private/ContentSourceProviders/FeaturePack/FeaturePackContentSource.h
+++ b/Engine/Source/Editor/AddContentDialog/Private/ContentSourceProviders/FeaturePack/FeaturePackContentSource.h
@@ -20,6 +20,8 @@ public:
 
 	virtual void InstallToProject(FString InstallPath) override;
 
+	virtual ~FFeaturePackContentSource();
+
 private:
 	void LoadPakFileToBuffer(FPakPlatformFile& PakPlatformFile, FString Path, TArray<uint8>& Buffer);
 
diff --git a/Engine/Source/Editor/AddContentDialog/Private/ContentSourceProviders/FeaturePack/FeaturePackContentSourceProvider.cpp b/Engine/Source/Editor/AddContentDialog/Private/ContentSourceProviders/FeaturePack/FeaturePackContentSourceProvider.cpp
--- a/Engine/Source/Editor/AddContentDialog/Private/ContentSourceProviders/FeaturePack/FeaturePackContentSourceProvider.cpp
+++ b/Engine/Source/Editor/AddContentDialog/Private/ContentSourceProviders/FeaturePack/FeaturePackContentSourceProvider.cpp
@@ -46,4 +46,8 @@ const TArray<TSharedRef<IContentSource>> FFeaturePackContentSourceProvider::GetC
 void FFeaturePackContentSourceProvider::SetContentSourcesChanged(FOnContentSourcesChanged OnContentSourcesChangedIn)
 {
 	OnContentSourcesChanged = OnContentSourcesChangedIn;
+}
+
+FFeaturePackContentSourceProvider::~FFeaturePackContentSourceProvider()
+{
 }
\ No newline at end of file
diff --git a/Engine/Source/Editor/AddContentDialog/Private/ContentSourceProviders/FeaturePack/FeaturePackContentSourceProvider.h b/Engine/Source/Editor/AddContentDialog/Private/ContentSourceProviders/FeaturePack/FeaturePackContentSourceProvider.h
--- a/Engine/Source/Editor/AddContentDialog/Private/ContentSourceProviders/FeaturePack/FeaturePackContentSourceProvider.h
+++ b/Engine/Source/Editor/AddContentDialog/Private/ContentSourceProviders/FeaturePack/FeaturePackContentSourceProvider.h
@@ -10,6 +10,8 @@ public:
 	virtual const TArray<TSharedRef<IContentSource>> GetContentSources() override;
 	virtual void SetContentSourcesChanged(FOnContentSourcesChanged OnContentSourcesChangedIn) override;
 
+	virtual ~FFeaturePackContentSourceProvider();
+
 private:
 	FOnContentSourcesChanged OnContentSourcesChanged;
 	TArray<TSharedRef<IContentSource>> ContentSources;
diff --git a/Engine/Source/Editor/AddContentDialog/Private/IContentSource.h b/Engine/Source/Editor/AddContentDialog/Private/IContentSource.h
--- a/Engine/Source/Editor/AddContentDialog/Private/IContentSource.h
+++ b/Engine/Source/Editor/AddContentDialog/Private/IContentSource.h
@@ -95,4 +95,6 @@ public:
 
 	/** Installs the content in the content source to the specific path. */
 	virtual void InstallToProject(FString InstallPath) = 0;
+
+	virtual ~IContentSource() { };
 };
diff --git a/Engine/Source/Editor/AddContentDialog/Private/IContentSourceProvider.h b/Engine/Source/Editor/AddContentDialog/Private/IContentSourceProvider.h
--- a/Engine/Source/Editor/AddContentDialog/Private/IContentSourceProvider.h
+++ b/Engine/Source/Editor/AddContentDialog/Private/IContentSourceProvider.h
@@ -14,6 +14,8 @@ public:
 	/** Gets the available content sources. */
 	virtual const TArray<TSharedRef<IContentSource>> GetContentSources() = 0;
 
-	/** Sets the delegate which will be executed when the avaialbe content sources change */
+	/** Sets the delegate which will be executed when the avaialble content sources change */
 	virtual void SetContentSourcesChanged(FOnContentSourcesChanged OnContentSourcesChangedIn) = 0;
+
+	virtual ~IContentSourceProvider() { }
 };
\ No newline at end of file
diff --git a/Engine/Source/Editor/AddContentDialog/Private/SAddContentDialog.cpp b/Engine/Source/Editor/AddContentDialog/Private/SAddContentDialog.cpp
--- a/Engine/Source/Editor/AddContentDialog/Private/SAddContentDialog.cpp
+++ b/Engine/Source/Editor/AddContentDialog/Private/SAddContentDialog.cpp
@@ -60,4 +60,6 @@ FReply SAddContentDialog::AddButtonClicked()
 
 SAddContentDialog::~SAddContentDialog()
 {
-}
\ No newline at end of file
+}
+
+#undef LOCTEXT_NAMESPACE
\ No newline at end of file
diff --git a/Engine/Source/Editor/AddContentDialog/Private/SAddContentWidget.cpp b/Engine/Source/Editor/AddContentDialog/Private/SAddContentWidget.cpp
--- a/Engine/Source/Editor/AddContentDialog/Private/SAddContentWidget.cpp
+++ b/Engine/Source/Editor/AddContentDialog/Private/SAddContentWidget.cpp
@@ -511,4 +511,6 @@ void SAddContentWidget::AddListChanged()
 
 SAddContentWidget::~SAddContentWidget()
 {
-}
\ No newline at end of file
+}
+
+#undef LOCTEXT_NAMESPACE
\ No newline at end of file
diff --git a/Engine/Source/Editor/AddContentDialog/Private/ViewModels/CategoryViewModel.cpp b/Engine/Source/Editor/AddContentDialog/Private/ViewModels/CategoryViewModel.cpp
--- a/Engine/Source/Editor/AddContentDialog/Private/ViewModels/CategoryViewModel.cpp
+++ b/Engine/Source/Editor/AddContentDialog/Private/ViewModels/CategoryViewModel.cpp
@@ -53,4 +53,6 @@ void FCategoryViewModel::Initialize()
 		IconBrush = FAddContentDialogStyle::Get().GetBrush("AddContentDialog.DefaultCategory");
 		break;
 	}
-}
\ No newline at end of file
+}
+
+#undef LOCTEXT_NAMESPACE
\ No newline at end of file
diff --git a/Engine/Source/Editor/AddContentDialog/Private/ViewModels/ContentSourceViewModel.cpp b/Engine/Source/Editor/AddContentDialog/Private/ViewModels/ContentSourceViewModel.cpp
--- a/Engine/Source/Editor/AddContentDialog/Private/ViewModels/ContentSourceViewModel.cpp
+++ b/Engine/Source/Editor/AddContentDialog/Private/ViewModels/ContentSourceViewModel.cpp
@@ -165,3 +165,5 @@ FLocalizedText FContentSourceViewModel::ChooseLocalizedText(TArray<FLocalizedTex
 	}
 	return Default;
 }
+
+#undef LOCTEXT_NAMESPACE
