AddContentDialog - Fix linker error.

diff --git a/Engine/Source/Editor/AddContentDialog/Private/ViewModels/CategoryViewModel.cpp b/Engine/Source/Editor/AddContentDialog/Private/ViewModels/CategoryViewModel.cpp
--- a/Engine/Source/Editor/AddContentDialog/Private/ViewModels/CategoryViewModel.cpp
+++ b/Engine/Source/Editor/AddContentDialog/Private/ViewModels/CategoryViewModel.cpp
@@ -26,11 +26,6 @@ const FSlateBrush* FCategoryViewModel::GetIconBrush() const
 	return IconBrush;
 }
 
-inline bool FCategoryViewModel::operator==(const FCategoryViewModel& Other) const
-{
-	return Other.Text.EqualTo(Text) && (Other.IconBrush == IconBrush);
-}
-
 uint32 FCategoryViewModel::GetTypeHash() const
 {
 	return (uint32)Category;
diff --git a/Engine/Source/Editor/AddContentDialog/Private/ViewModels/CategoryViewModel.h b/Engine/Source/Editor/AddContentDialog/Private/ViewModels/CategoryViewModel.h
--- a/Engine/Source/Editor/AddContentDialog/Private/ViewModels/CategoryViewModel.h
+++ b/Engine/Source/Editor/AddContentDialog/Private/ViewModels/CategoryViewModel.h
@@ -16,7 +16,10 @@ public:
 	/** Gets the brush which should be used to draw the icon for the category. */
 	const FSlateBrush* GetIconBrush() const;
 
-	inline bool operator==(const FCategoryViewModel& Other) const;
+	inline bool operator==(const FCategoryViewModel& Other) const
+	{
+		return Other.Text.EqualTo(Text) && (Other.IconBrush == IconBrush);
+	}
 
 	uint32 GetTypeHash() const;
 
