AddContentDialog - Upgrade to the nice widget carousel that the launcher uses.

diff --git a/Engine/Source/Editor/AddContentDialog/AddContentDialog.Build.cs b/Engine/Source/Editor/AddContentDialog/AddContentDialog.Build.cs
--- a/Engine/Source/Editor/AddContentDialog/AddContentDialog.Build.cs
+++ b/Engine/Source/Editor/AddContentDialog/AddContentDialog.Build.cs
@@ -28,6 +28,12 @@ public class AddContentDialog : ModuleRules
 			}
 		);
 
+		PrivateDependencyModuleNames.AddRange(
+			new string[] {
+				"WidgetCarousel"
+			}
+		);
+
 		DynamicallyLoadedModuleNames.AddRange(
 			new string[] {
 			}
diff --git a/Engine/Source/Editor/AddContentDialog/Private/AddContentDialogModule.cpp b/Engine/Source/Editor/AddContentDialog/Private/AddContentDialogModule.cpp
--- a/Engine/Source/Editor/AddContentDialog/Private/AddContentDialogModule.cpp
+++ b/Engine/Source/Editor/AddContentDialog/Private/AddContentDialogModule.cpp
@@ -6,6 +6,7 @@
 #include "FeaturePackContentSourceProvider.h"
 #include "ModuleManager.h"
 #include "SDockTab.h"
+#include "WidgetCarouselModule.h"
 
 #define LOCTEXT_NAMESPACE "AddContentDialog"
 
@@ -15,6 +16,7 @@ class FAddContentDialogModule : public IAddContentDialogModule
 public:
 	virtual void StartupModule() override
 	{
+		FModuleManager::LoadModuleChecked<FWidgetCarouselModule>("WidgetCarousel");
 		ContentSourceProviderManager = TSharedPtr<FContentSourceProviderManager>(new FContentSourceProviderManager());
 		ContentSourceProviderManager->RegisterContentSourceProvider(MakeShareable(new FFeaturePackContentSourceProvider()));
 		FAddContentDialogStyle::Initialize();
diff --git a/Engine/Source/Editor/AddContentDialog/Private/SAddContentWidget.cpp b/Engine/Source/Editor/AddContentDialog/Private/SAddContentWidget.cpp
--- a/Engine/Source/Editor/AddContentDialog/Private/SAddContentWidget.cpp
+++ b/Engine/Source/Editor/AddContentDialog/Private/SAddContentWidget.cpp
@@ -3,6 +3,7 @@
 #include "AddContentDialogPCH.h"
 
 #include "SSearchBox.h"
+#include "SWidgetCarouselWithNavigation.h"
 
 #define LOCTEXT_NAMESPACE "AddContentDialog"
 
@@ -285,68 +286,9 @@ TSharedRef<SWidget> SAddContentWidget::CreateContentSourceDetail(TSharedPtr<FCon
 
 TSharedRef<SWidget> SAddContentWidget::CreateScreenshotCarousel(TSharedPtr<FContentSourceViewModel> ContentSource)
 {
-	// Previous button and image.
-	TSharedRef<SButton> PreviousButton = 
-		SNew(SButton)
-		.ButtonStyle(FAddContentDialogStyle::Get(), "AddContentDialog.BlankButton")
-		.Visibility(ContentSource->GetScreenshotBrushes()->Num() > 1 ? EVisibility::Visible : EVisibility::Hidden)
-		.OnClicked(this, &SAddContentWidget::PreviousScreenshotButtonClicked, ContentSource)
-		.HAlign(EHorizontalAlignment::HAlign_Left)
-		.VAlign(EVerticalAlignment::VAlign_Center);
-
-	TSharedRef<SImage> PreviousImage = 
-		SNew(SImage)
-		.Image(FAddContentDialogStyle::Get().GetBrush("AddContentDialog.LeftArrow"))
-		.Visibility_Lambda([PreviousButton]()->EVisibility
-		{
-			return PreviousButton->IsHovered() ? EVisibility::Visible : EVisibility::Hidden;
-		});
-
-	PreviousButton->SetContent(PreviousImage);
-
-	// Next button and image.
-	TSharedRef<SButton> NextButton = 
-		SNew(SButton)
-		.ButtonStyle(FAddContentDialogStyle::Get(), "AddContentDialog.BlankButton")
-		.Visibility(ContentSource->GetScreenshotBrushes()->Num() > 1 ? EVisibility::Visible : EVisibility::Hidden)
-		.OnClicked(this, &SAddContentWidget::NextScreenshotButtonClicked, ContentSource)
-		.HAlign(EHorizontalAlignment::HAlign_Right)
-		.VAlign(EVerticalAlignment::VAlign_Center);
-
-	TSharedRef<SImage> NextImage = 
-		SNew(SImage)
-		.Image(FAddContentDialogStyle::Get().GetBrush("AddContentDialog.RightArrow"))
-		.Visibility_Lambda([NextButton]()->EVisibility
-		{
-			return NextButton->IsHovered() ? EVisibility::Visible : EVisibility::Hidden;
-		});
-
-	NextButton->SetContent(NextImage);
-
-	// Image and button overlay.
-	return SNew(SOverlay)
-
-	+ SOverlay::Slot()
-	.HAlign(HAlign_Center)
-	.VAlign(VAlign_Center)
-	[
-		SNew(SImage)
-		.Image(ContentSource.ToSharedRef(), &FContentSourceViewModel::GetSelectedScreenshotBrush)
-	]
-
-	+ SOverlay::Slot()
-	.Padding(0)
-	[
-		SNew(SHorizontalBox)
-		+ SHorizontalBox::Slot()
-		[
-			PreviousButton
-		]
-		+ SHorizontalBox::Slot()
-		[
-			NextButton
-		]
-	];
+	return SNew(SWidgetCarouselWithNavigation<TSharedPtr<FSlateBrush>>)
+	.OnGenerateWidget(this, &SAddContentWidget::CreateScreenshotWidget)
+	.WidgetItemsSource(ContentSource->GetScreenshotBrushes());
 }
 
 TSharedRef<SWidget> SAddContentWidget::CreateAddListView()
@@ -465,16 +407,10 @@ FReply SAddContentWidget::ContentSourceTileDragDetected(const FGeometry& MyGeome
 	return FReply::Unhandled();
 }
 
-FReply SAddContentWidget::PreviousScreenshotButtonClicked(TSharedPtr<FContentSourceViewModel> ContentSource)
-{
-	ContentSource->SelectPreviousScreenshotBrush();
-	return FReply::Handled();
-}
-
-FReply SAddContentWidget::NextScreenshotButtonClicked(TSharedPtr<FContentSourceViewModel> ContentSource)
+TSharedRef<SWidget> SAddContentWidget::CreateScreenshotWidget(TSharedPtr<FSlateBrush> ScreenshotBrush)
 {
-	ContentSource->SelectNextScreenshotBrush();
-	return FReply::Handled();
+	return SNew(SImage)
+	.Image(ScreenshotBrush.Get());
 }
 
 FReply SAddContentWidget::RemoveAddedContentSourceClicked(TSharedPtr<FContentSourceViewModel> ContentSource)
diff --git a/Engine/Source/Editor/AddContentDialog/Private/SAddContentWidget.h b/Engine/Source/Editor/AddContentDialog/Private/SAddContentWidget.h
--- a/Engine/Source/Editor/AddContentDialog/Private/SAddContentWidget.h
+++ b/Engine/Source/Editor/AddContentDialog/Private/SAddContentWidget.h
@@ -40,10 +40,13 @@ private:
 	/** Creates a widget carousel for displaying the set of screenshots for a specific content source. */
 	TSharedRef<SWidget> CreateScreenshotCarousel(TSharedPtr<FContentSourceViewModel> ContentSource);
 
-	/** Creates a list view represnting content which the user has selected for addition to the project. */
+	/** Creates the widget the displays a screenshot in the screenshot carousel. */
+	TSharedRef<SWidget> SAddContentWidget::CreateScreenshotWidget(TSharedPtr<FSlateBrush> ScreenshotBrush);
+
+	/** Creates a list view representing content which the user has selected for addition to the project. */
 	TSharedRef<SWidget> CreateAddListView();
 	
-	/** Creates a row for the list view represnting content selected for addition to the project */
+	/** Creates a row for the list view representing content selected for addition to the project */
 	TSharedRef<ITableRow> CreateAddContentSourceRow(TSharedPtr<FContentSourceViewModel> ContentSource, const TSharedRef<STableViewBase>& OwnerTable);
 
 
@@ -63,12 +66,6 @@ private:
 	/** Handles an item from the content source tile view being dragged. */
 	FReply ContentSourceTileDragDetected(const FGeometry& MyGeometry, const FPointerEvent& MouseEvent);
 
-	/** Handles the previous screenshot button in the screenshot carousel being pressed. */
-	FReply PreviousScreenshotButtonClicked(TSharedPtr<FContentSourceViewModel> ContentSource);
-
-	/** Handles the next screenshot button in the screenshot carousel being pressed */
-	FReply NextScreenshotButtonClicked(TSharedPtr<FContentSourceViewModel> ContentSource);
-
 	/** Handles the remove button being pressed on a row in the list of content selected for addition to the project. */
 	FReply RemoveAddedContentSourceClicked(TSharedPtr<FContentSourceViewModel> ContentSource);
 
