AddContentDialog - Remove unused methods from the content source view model.

diff --git a/Engine/Source/Editor/AddContentDialog/Private/ViewModels/ContentSourceViewModel.cpp b/Engine/Source/Editor/AddContentDialog/Private/ViewModels/ContentSourceViewModel.cpp
--- a/Engine/Source/Editor/AddContentDialog/Private/ViewModels/ContentSourceViewModel.cpp
+++ b/Engine/Source/Editor/AddContentDialog/Private/ViewModels/ContentSourceViewModel.cpp
@@ -56,48 +56,6 @@ TArray<TSharedPtr<FSlateBrush>>* FContentSourceViewModel::GetScreenshotBrushes()
 	return &ScreenshotBrushes;
 }
 
-
-const FSlateBrush* FContentSourceViewModel::GetSelectedScreenshotBrush() const
-{
-	return SelectedScreenshotBrush.Get();
-}
-
-void FContentSourceViewModel::SelectNextScreenshotBrush()
-{
-	if (ScreenshotBrushes.Num() > 1)
-	{
-		int currentIndex = ScreenshotBrushes.IndexOfByKey(SelectedScreenshotBrush);
-		int newIndex = currentIndex + 1;
-		if (newIndex < 0)
-		{
-			newIndex += ScreenshotBrushes.Num();
-		}
-		else if (newIndex >= ScreenshotBrushes.Num())
-		{
-			newIndex -= ScreenshotBrushes.Num();
-		}
-		SelectedScreenshotBrush = ScreenshotBrushes[newIndex];
-	}
-}
-
-void FContentSourceViewModel::SelectPreviousScreenshotBrush()
-{
-	if (ScreenshotBrushes.Num() > 1)
-	{
-		int currentIndex = ScreenshotBrushes.IndexOfByKey(SelectedScreenshotBrush);
-		int newIndex = currentIndex - 1;
-		if (newIndex < 0)
-		{
-			newIndex += ScreenshotBrushes.Num();
-		}
-		else if (newIndex >= ScreenshotBrushes.Num())
-		{
-			newIndex -= ScreenshotBrushes.Num();
-		}
-		SelectedScreenshotBrush = ScreenshotBrushes[newIndex];
-	}
-}
-
 void FContentSourceViewModel::SetupBrushes()
 {
 	FString IconBrushName = GetName().ToString() + "_" + ContentSource->GetIconData()->GetName();
@@ -108,10 +66,6 @@ void FContentSourceViewModel::SetupBrushes()
 		FString ScreenshotBrushName = GetName().ToString() + "_" + ScreenshotData->GetName();
 		ScreenshotBrushes.Add(CreateBrushFromRawData(FName(*ScreenshotBrushName), *ScreenshotData->GetData()));
 	}
-	if (ScreenshotBrushes.Num())
-	{
-		SelectedScreenshotBrush = ScreenshotBrushes[0];
-	}
 }
 
 TSharedPtr<FSlateDynamicImageBrush> FContentSourceViewModel::CreateBrushFromRawData(FName ResourceName, const TArray< uint8 >& RawData) const
diff --git a/Engine/Source/Editor/AddContentDialog/Private/ViewModels/ContentSourceViewModel.h b/Engine/Source/Editor/AddContentDialog/Private/ViewModels/ContentSourceViewModel.h
--- a/Engine/Source/Editor/AddContentDialog/Private/ViewModels/ContentSourceViewModel.h
+++ b/Engine/Source/Editor/AddContentDialog/Private/ViewModels/ContentSourceViewModel.h
@@ -27,17 +27,6 @@ public:
 	/** Gets an array or brushes which should be used to display screenshots for this content source. */
 	TArray<TSharedPtr<FSlateBrush>>* GetScreenshotBrushes();
 
-	/** Gets the brush for the currently selected screenshot. */
-	const FSlateBrush* GetSelectedScreenshotBrush() const;
-
-	/** Selects the next available screenshot.  When at the last screenshot this will start again at
-		the beginning of the array. */
-	void SelectNextScreenshotBrush();
-
-	/** Selects the previous available screenshot.  When at the first screenshot this will start again at
-		the end of the array. */
-	void SelectPreviousScreenshotBrush();
-
 private:
 	/** Sets up brushes from the images data supplied by the IContentSource. */
 	void SetupBrushes();
@@ -58,9 +47,6 @@ private:
 	/** An array or brushes which should be used to display screenshots for this content source. */
 	TArray<TSharedPtr<FSlateBrush>> ScreenshotBrushes;
 
-	/** The brush for the currently selected screenshot. */
-	TSharedPtr<FSlateBrush> SelectedScreenshotBrush;
-
 	/** The view model for the category for this content source. */
 	FCategoryViewModel Category;
 
