CIS fix

diff --git a/Engine/Source/Developer/WidgetCarousel/Private/SCarouselNavigationBar.cpp b/Engine/Source/Developer/WidgetCarousel/Private/SCarouselNavigationBar.cpp
--- a/Engine/Source/Developer/WidgetCarousel/Private/SCarouselNavigationBar.cpp
+++ b/Engine/Source/Developer/WidgetCarousel/Private/SCarouselNavigationBar.cpp
@@ -99,4 +99,6 @@ FVector2D SCarouselNavigationBar::GetPlaceHolderPosition() const
 	NewPos -= CurrentSlideAmount.Get() * WidgetSize;
 
 	return FVector2D(NewPos, 0.f);
-}
\ No newline at end of file
+}
+
+#undef LOCTEXT_NAMESPACE
diff --git a/Engine/Source/Developer/WidgetCarousel/Public/SCarouselNavigationBar.h b/Engine/Source/Developer/WidgetCarousel/Public/SCarouselNavigationBar.h
--- a/Engine/Source/Developer/WidgetCarousel/Public/SCarouselNavigationBar.h
+++ b/Engine/Source/Developer/WidgetCarousel/Public/SCarouselNavigationBar.h
@@ -51,4 +51,6 @@ private:
 	int32 ItemCount;
 	TAttribute<int32> CurrentItemIndex;
 	TAttribute<float> CurrentSlideAmount;
-};
\ No newline at end of file
+};
+
+#undef LOCTEXT_NAMESPACE
diff --git a/Engine/Source/Editor/AddContentDialog/Private/SAddContentWidget.h b/Engine/Source/Editor/AddContentDialog/Private/SAddContentWidget.h
--- a/Engine/Source/Editor/AddContentDialog/Private/SAddContentWidget.h
+++ b/Engine/Source/Editor/AddContentDialog/Private/SAddContentWidget.h
@@ -41,7 +41,7 @@ private:
 	TSharedRef<SWidget> CreateScreenshotCarousel(TSharedPtr<FContentSourceViewModel> ContentSource);
 
 	/** Creates the widget the displays a screenshot in the screenshot carousel. */
-	TSharedRef<SWidget> SAddContentWidget::CreateScreenshotWidget(TSharedPtr<FSlateBrush> ScreenshotBrush);
+	TSharedRef<SWidget> CreateScreenshotWidget(TSharedPtr<FSlateBrush> ScreenshotBrush);
 
 	/** Creates a list view representing content which the user has selected for addition to the project. */
 	TSharedRef<SWidget> CreateAddListView();
