Added 4 Blueprint Feature packs

diff --git a/Engine/Source/Editor/AddContentDialog/Private/ViewModels/CategoryViewModel.cpp b/Engine/Source/Editor/AddContentDialog/Private/ViewModels/CategoryViewModel.cpp
--- a/Engine/Source/Editor/AddContentDialog/Private/ViewModels/CategoryViewModel.cpp
+++ b/Engine/Source/Editor/AddContentDialog/Private/ViewModels/CategoryViewModel.cpp
@@ -43,6 +43,10 @@ void FCategoryViewModel::Initialize()
 		Text = LOCTEXT("CodeFeature", "C++ Feature");
 		IconBrush = FAddContentDialogStyle::Get().GetBrush("AddContentDialog.CodeFeatureCategory");
 		break;
+	case EContentSourceCategory::Content:
+		Text = LOCTEXT("ContentPacks", "Content Packs");
+		IconBrush = FAddContentDialogStyle::Get().GetBrush("AddContentDialog.ContentPacks");
+		break;
 	default:
 		Text = LOCTEXT("Miscellaneous", "Miscellaneous");
 		IconBrush = FAddContentDialogStyle::Get().GetBrush("AddContentDialog.DefaultCategory");
