Fixed feature pack import issues created by reverting startercontent reversion
Readded 'content' feature pack type

diff --git a/Engine/Source/Editor/AddContentDialog/Private/ContentSourceProviders/FeaturePack/FeaturePackContentSource.cpp b/Engine/Source/Editor/AddContentDialog/Private/ContentSourceProviders/FeaturePack/FeaturePackContentSource.cpp
--- a/Engine/Source/Editor/AddContentDialog/Private/ContentSourceProviders/FeaturePack/FeaturePackContentSource.cpp
+++ b/Engine/Source/Editor/AddContentDialog/Private/ContentSourceProviders/FeaturePack/FeaturePackContentSource.cpp
@@ -60,6 +60,10 @@ FFeaturePackContentSource::FFeaturePackContentSource(FString InFeaturePackPath)
 	{
 		Category = EContentSourceCategory::BlueprintFeature;
 	}
+	else if (CategoryString == "Content")
+	{
+		Category = EContentSourceCategory::Content;
+	}
 	else
 	{
 		Category = EContentSourceCategory::Unknown;
diff --git a/Engine/Source/Editor/AddContentDialog/Private/IContentSource.h b/Engine/Source/Editor/AddContentDialog/Private/IContentSource.h
--- a/Engine/Source/Editor/AddContentDialog/Private/IContentSource.h
+++ b/Engine/Source/Editor/AddContentDialog/Private/IContentSource.h
@@ -7,6 +7,7 @@ enum class EContentSourceCategory
 {
 	BlueprintFeature,
 	CodeFeature,
+	Content,
 	Unknown
 };
 
diff --git a/Engine/Source/Editor/AddContentDialog/Private/ViewModels/CategoryViewModel.cpp b/Engine/Source/Editor/AddContentDialog/Private/ViewModels/CategoryViewModel.cpp
--- a/Engine/Source/Editor/AddContentDialog/Private/ViewModels/CategoryViewModel.cpp
+++ b/Engine/Source/Editor/AddContentDialog/Private/ViewModels/CategoryViewModel.cpp
@@ -43,6 +43,10 @@ void FCategoryViewModel::Initialize()
 		Text = LOCTEXT("CodeFeature", "C++ Feature");
 		IconBrush = FAddContentDialogStyle::Get().GetBrush("AddContentDialog.CodeFeatureCategory");
 		break;
+	case EContentSourceCategory::Content:
+		Text = LOCTEXT("ContentPacks", "Content Packs");
+		IconBrush = FAddContentDialogStyle::Get().GetBrush("AddContentDialog.ContentPacks");
+		break;
 	default:
 		Text = LOCTEXT("Miscellaneous", "Miscellaneous");
 		IconBrush = FAddContentDialogStyle::Get().GetBrush("AddContentDialog.DefaultCategory");
diff --git a/Engine/Source/Runtime/Core/Private/Misc/Paths.cpp b/Engine/Source/Runtime/Core/Private/Misc/Paths.cpp
--- a/Engine/Source/Runtime/Core/Private/Misc/Paths.cpp
+++ b/Engine/Source/Runtime/Core/Private/Misc/Paths.cpp
@@ -559,6 +559,7 @@ bool FPaths::IsRelative(const FString& InPath)
 							InPath.StartsWith(TEXT("//"))	||												// Equivalent to "\\", considering normalization replaces "\\" with "//".
 							InPath.StartsWith(TEXT("\\"))	||												// Root of the current directory on Windows
 							InPath.StartsWith(TEXT("/"))	||												// Root of the current directory on Windows, root on UNIX-likes.
+							InPath.StartsWith(TEXT("root:/")) ||											// Feature packs use this
 							(InPath.Len() >= 2 && FChar::IsAlpha(InPath[0]) && InPath[1] == TEXT(':'));	// Starts with "<DriveLetter>:"
 
 	return !IsRooted;
