AddContentDialog - Minor additions.
+ Refresh feature packs automatically when they change on disk.
+ Remove drag and drop and batch install, and replace it with a single button for installing in the details section.
+ Don't allow multiple copies of the dialog to be opened.
+ Fix a missing brush visible when resizing the dialog.
+ Fix an issue where brushes in the dialog could get deleted resulting in white images.

diff --git a/Engine/Source/Developer/WidgetCarousel/Private/WidgetCarouselStyle.cpp b/Engine/Source/Developer/WidgetCarousel/Private/WidgetCarouselStyle.cpp
--- a/Engine/Source/Developer/WidgetCarousel/Private/WidgetCarouselStyle.cpp
+++ b/Engine/Source/Developer/WidgetCarousel/Private/WidgetCarouselStyle.cpp
@@ -113,7 +113,7 @@ TSharedRef< FSlateStyleSet > FWidgetCarouselModuleStyle::Create()
 		.SetHovered(IMAGE_BRUSH("CarouselNavRight", FVector2D(80.0f, 20.0f), PrimaryCallToActionColor))
 			.SetPressed(IMAGE_BRUSH("CarouselNavRight", FVector2D(80.0f, 20.0f), PrimaryCallToActionColorPressed))));
 
-	Style->Set("Black .6", FLinearColor(0, 0, 0, 0.6f));
+	Style->Set("WidgetBackground", new FSlateColorBrush(FLinearColor(0, 0, 0, 0.6f)));
 
 	return Style;
 }
diff --git a/Engine/Source/Developer/WidgetCarousel/Public/SWidgetCarouselWithNavigation.h b/Engine/Source/Developer/WidgetCarousel/Public/SWidgetCarouselWithNavigation.h
--- a/Engine/Source/Developer/WidgetCarousel/Public/SWidgetCarouselWithNavigation.h
+++ b/Engine/Source/Developer/WidgetCarousel/Public/SWidgetCarouselWithNavigation.h
@@ -50,7 +50,7 @@ public:
 					.VAlign(VAlign_Center)
 					[
 						SNew(SBorder)
-						.BorderImage(FWidgetCarouselModuleStyle::Get().GetBrush("Black .6"))
+						.BorderImage(FWidgetCarouselModuleStyle::Get().GetBrush("WidgetBackground"))
 						.Padding(0)
 						[
 							SNew(SScaleBox)
diff --git a/Engine/Source/Editor/AddContentDialog/Private/ContentSourceProviders/FeaturePack/FeaturePackContentSourceProvider.cpp b/Engine/Source/Editor/AddContentDialog/Private/ContentSourceProviders/FeaturePack/FeaturePackContentSourceProvider.cpp
--- a/Engine/Source/Editor/AddContentDialog/Private/ContentSourceProviders/FeaturePack/FeaturePackContentSourceProvider.cpp
+++ b/Engine/Source/Editor/AddContentDialog/Private/ContentSourceProviders/FeaturePack/FeaturePackContentSourceProvider.cpp
@@ -5,6 +5,8 @@
 #include "FeaturePackContentSourceProvider.h"
 #include "FeaturePackContentSource.h"
 
+#include "ModuleManager.h"
+
 class FFillArrayDirectoryVisitor : public IPlatformFile::FDirectoryVisitor
 {
 public:
@@ -27,19 +29,9 @@ public:
 
 FFeaturePackContentSourceProvider::FFeaturePackContentSourceProvider()
 {
-	//TODO: Add a directory watcher?
-	FString FeaturePackPath = FPaths::Combine(*FPaths::RootDir(), TEXT("FeaturePacks"), TEXT("Packed"));
-	IPlatformFile &PlatformFile = FPlatformFileManager::Get().GetPlatformFile();
-	FFillArrayDirectoryVisitor DirectoryVisitor;
-	PlatformFile.IterateDirectory(*FeaturePackPath, DirectoryVisitor);
-	for (auto FeaturePackFile : DirectoryVisitor.Files)
-	{
-		TUniquePtr<FFeaturePackContentSource> NewContentSource = MakeUnique<FFeaturePackContentSource>(FeaturePackFile);
-		if(NewContentSource->IsDataValid())
-		{
-			ContentSources.Add(MakeShareable(NewContentSource.Release()));
-		}
-	}
+	FeaturePackPath = FPaths::Combine(*FPaths::RootDir(), TEXT("FeaturePacks"), TEXT("Packed"));
+	StartUpDirectoryWatcher();
+	RefreshFeaturePacks();
 }
 
 const TArray<TSharedRef<IContentSource>> FFeaturePackContentSourceProvider::GetContentSources()
@@ -52,6 +44,52 @@ void FFeaturePackContentSourceProvider::SetContentSourcesChanged(FOnContentSourc
 	OnContentSourcesChanged = OnContentSourcesChangedIn;
 }
 
+void FFeaturePackContentSourceProvider::StartUpDirectoryWatcher()
+{
+	FDirectoryWatcherModule& DirectoryWatcherModule = FModuleManager::LoadModuleChecked<FDirectoryWatcherModule>(TEXT("DirectoryWatcher"));
+	IDirectoryWatcher* DirectoryWatcher = DirectoryWatcherModule.Get();
+	if (DirectoryWatcher)
+	{
+		// If the path doesn't exist on disk, make it so the watcher will work.
+		IFileManager::Get().MakeDirectory(*FeaturePackPath);
+		DirectoryChangedDelegate = IDirectoryWatcher::FDirectoryChanged::CreateRaw(this, &FFeaturePackContentSourceProvider::OnFeaturePackDirectoryChanged);
+		DirectoryWatcher->RegisterDirectoryChangedCallback(FeaturePackPath, DirectoryChangedDelegate);
+	}
+}
+
+void FFeaturePackContentSourceProvider::ShutDownDirectoryWatcher()
+{
+	FDirectoryWatcherModule& DirectoryWatcherModule = FModuleManager::LoadModuleChecked<FDirectoryWatcherModule>( TEXT( "DirectoryWatcher" ) );
+	IDirectoryWatcher* DirectoryWatcher = DirectoryWatcherModule.Get();
+	if ( DirectoryWatcher )
+	{
+		DirectoryWatcher->UnregisterDirectoryChangedCallback(FeaturePackPath, DirectoryChangedDelegate);
+	}
+}
+
+void FFeaturePackContentSourceProvider::OnFeaturePackDirectoryChanged( const TArray<FFileChangeData>& FileChanges )
+{
+	RefreshFeaturePacks();
+}
+
+void FFeaturePackContentSourceProvider::RefreshFeaturePacks()
+{
+	ContentSources.Empty();
+	IPlatformFile &PlatformFile = FPlatformFileManager::Get().GetPlatformFile();
+	FFillArrayDirectoryVisitor DirectoryVisitor;
+	PlatformFile.IterateDirectory( *FeaturePackPath, DirectoryVisitor );
+	for ( auto FeaturePackFile : DirectoryVisitor.Files )
+	{
+		TUniquePtr<FFeaturePackContentSource> NewContentSource = MakeUnique<FFeaturePackContentSource>( FeaturePackFile );
+		if ( NewContentSource->IsDataValid() )
+		{
+			ContentSources.Add( MakeShareable( NewContentSource.Release() ) );
+		}
+	}
+	OnContentSourcesChanged.ExecuteIfBound();
+}
+
 FFeaturePackContentSourceProvider::~FFeaturePackContentSourceProvider()
 {
+	ShutDownDirectoryWatcher();
 }
\ No newline at end of file
diff --git a/Engine/Source/Editor/AddContentDialog/Private/ContentSourceProviders/FeaturePack/FeaturePackContentSourceProvider.h b/Engine/Source/Editor/AddContentDialog/Private/ContentSourceProviders/FeaturePack/FeaturePackContentSourceProvider.h
--- a/Engine/Source/Editor/AddContentDialog/Private/ContentSourceProviders/FeaturePack/FeaturePackContentSourceProvider.h
+++ b/Engine/Source/Editor/AddContentDialog/Private/ContentSourceProviders/FeaturePack/FeaturePackContentSourceProvider.h
@@ -1,5 +1,7 @@
 // Copyright 1998-2015 Epic Games, Inc. All Rights Reserved.
 
+#include "DirectoryWatcherModule.h"
+
 #pragma once
 
 /** A content source provider for available content upacks. */
@@ -13,6 +15,28 @@ public:
 	virtual ~FFeaturePackContentSourceProvider();
 
 private:
+	/** Starts the directory watcher for the feature pack directory. */
+	void StartUpDirectoryWatcher();
+
+	/** Shuts down the directory watcher for the feature pack directory. */
+	void ShutDownDirectoryWatcher();
+
+	/** Delegate to handle whenever the feature pack directory changes on disk. */
+	void OnFeaturePackDirectoryChanged( const TArray<FFileChangeData>& FileChanges );
+
+	/** Rebuilds the feature pack array and calls the change notification delegate. */
+	void RefreshFeaturePacks();
+
+	/** The path on disk to the directory containing the feature packs. */
+	FString FeaturePackPath;
+
+	/** The delegate which gets called when the feature pack directory changes. This reference
+		is kept so that it can be unregistered correctly. */
+	IDirectoryWatcher::FDirectoryChanged DirectoryChangedDelegate;
+
+	/** A delegate which gets called whenever the array of content sources changes. */
 	FOnContentSourcesChanged OnContentSourcesChanged;
+
+	/** An array of the available content sources. */
 	TArray<TSharedRef<IContentSource>> ContentSources;
 };
\ No newline at end of file
diff --git a/Engine/Source/Editor/AddContentDialog/Private/SAddContentDialog.cpp b/Engine/Source/Editor/AddContentDialog/Private/SAddContentDialog.cpp
--- a/Engine/Source/Editor/AddContentDialog/Private/SAddContentDialog.cpp
+++ b/Engine/Source/Editor/AddContentDialog/Private/SAddContentDialog.cpp
@@ -11,7 +11,7 @@ void SAddContentDialog::Construct(const FArguments& InArgs)
 	SWindow::Construct(SWindow::FArguments()
 		.Title(LOCTEXT("AddContentDialogTitle", "Add Content to the Project"))
 		.SizingRule(ESizingRule::UserSized)
-		.ClientSize(FVector2D(900, 700))
+		.ClientSize(FVector2D(900, 500))
 		.SupportsMinimize(false)
 		.SupportsMaximize(false)
 		[
@@ -23,41 +23,13 @@ void SAddContentDialog::Construct(const FArguments& InArgs)
 
 				// Add content widget.
 				+ SVerticalBox::Slot()
-				.Padding(FMargin(0, 0, 0, 15))
 				[
 					AddContentWidget.ToSharedRef()
 				]
-
-				// Add button.
-				+ SVerticalBox::Slot()
-				.AutoHeight()
-				.HAlign(EHorizontalAlignment::HAlign_Right)
-				[
-					SNew(SButton)
-					.OnClicked(this, &SAddContentDialog::AddButtonClicked)
-					.IsEnabled_Lambda([this]()->bool
-					{
-						return AddContentWidget->GetContentSourcesToAdd()->Num() > 0;
-					})
-					[
-						SNew(STextBlock)
-						.Text(LOCTEXT("AddSelectedContentToProject", "Add the selected content to the project"))
-					]
-				]
 			]
 		]);
 }
 
-FReply SAddContentDialog::AddButtonClicked()
-{
-	for (TSharedPtr<IContentSource> ContentSource : *AddContentWidget->GetContentSourcesToAdd())
-	{
-		ContentSource->InstallToProject("/Game");
-	}
-	RequestDestroyWindow();
-	return FReply::Handled();
-}
-
 SAddContentDialog::~SAddContentDialog()
 {
 }
diff --git a/Engine/Source/Editor/AddContentDialog/Private/SAddContentDialog.h b/Engine/Source/Editor/AddContentDialog/Private/SAddContentDialog.h
--- a/Engine/Source/Editor/AddContentDialog/Private/SAddContentDialog.h
+++ b/Engine/Source/Editor/AddContentDialog/Private/SAddContentDialog.h
@@ -15,10 +15,6 @@ public:
 
 	void Construct(const FArguments& InArgs);
 
-private:
-	/** Handles the add content to project button being clicked. */
-	FReply AddButtonClicked();
-
 private:
 	/** The widget representing available content and which content the user has selected. */
 	TSharedPtr<SAddContentWidget> AddContentWidget;
diff --git a/Engine/Source/Editor/AddContentDialog/Private/SAddContentWidget.cpp b/Engine/Source/Editor/AddContentDialog/Private/SAddContentWidget.cpp
--- a/Engine/Source/Editor/AddContentDialog/Private/SAddContentWidget.cpp
+++ b/Engine/Source/Editor/AddContentDialog/Private/SAddContentWidget.cpp
@@ -7,75 +7,6 @@
 
 #define LOCTEXT_NAMESPACE "AddContentDialog"
 
-class SContentSourceDragDropHandler : public SCompoundWidget
-{
-public:
-	SLATE_BEGIN_ARGS(SContentSourceDragDropHandler) {}
-	SLATE_DEFAULT_SLOT(FArguments, Content)
-	SLATE_ARGUMENT(TWeakPtr<FAddContentWidgetViewModel>, ViewModel)
-	SLATE_END_ARGS()
-
-	void Construct(const FArguments& InArgs)
-	{
-		ViewModel = InArgs._ViewModel;
-		bIsDragOver = false;
-
-		this->ChildSlot
-		[
-			SNew(SBorder)
-			.BorderImage(FEditorStyle::GetBrush("WhiteBrush"))
-			.BorderBackgroundColor(this, &SContentSourceDragDropHandler::GetBackgroundColor)
-			[
-				InArgs._Content.Widget
-			]
-		];
-	}
-
-	~SContentSourceDragDropHandler() {}
-
-	FReply OnDrop(const FGeometry& MyGeometry, const FDragDropEvent& DragDropEvent) override
-	{
-		bIsDragOver = false;
-
-		TSharedPtr<FAddContentWidgetViewModel> PinnedViewModel = ViewModel.Pin();
-		TSharedPtr<FContentSourceDragDropOp> ContentSourceDragDropOp = DragDropEvent.GetOperationAs<FContentSourceDragDropOp>();
-		if (PinnedViewModel.IsValid() && ContentSourceDragDropOp.IsValid())
-		{
-			PinnedViewModel->AppendToAddList(ContentSourceDragDropOp->GetContentSource());
-			return FReply::Handled();
-		}
-		return FReply::Unhandled();
-	}
-
-	virtual void OnDragEnter(const FGeometry& MyGeometry, const FDragDropEvent& DragDropEvent) override
-	{
-		bIsDragOver = true;
-	}
-
-	virtual void OnDragLeave(const FDragDropEvent& DragDropEvent) override
-	{
-		bIsDragOver = false;
-	}
-
-private:
-	FSlateColor GetBackgroundColor() const
-	{
-		if (bIsDragOver)
-		{
-			return FLinearColor(1.0f, 0.6f, 0.1f, 0.9f);
-		}
-		else
-		{
-			return FSlateApplication::Get().IsDragDropping() && FSlateApplication::Get().GetDragDroppingContent()->IsOfType<FContentSourceDragDropOp>() ?
-				FLinearColor(0.1f, 0.1f, 0.1f, 0.9f) : FLinearColor::Transparent;
-		}
-	}
-
-private:
-	TWeakPtr<FAddContentWidgetViewModel> ViewModel;
-	bool bIsDragOver;
-};
-
 void SAddContentWidget::Construct(const FArguments& InArgs)
 {
 	ViewModel = FAddContentWidgetViewModel::CreateShared();
@@ -85,8 +16,6 @@ void SAddContentWidget::Construct(const FArguments& InArgs)
 		this, &SAddContentWidget::ContentSourcesChanged));
 	ViewModel->SetOnSelectedContentSourceChanged(FAddContentWidgetViewModel::FOnSelectedContentSourceChanged::CreateSP(
 		this, &SAddContentWidget::SelectedContentSourceChanged));
-	ViewModel->SetOnAddListChanged(FAddContentWidgetViewModel::FOnAddListChanged::CreateSP(
-		this, &SAddContentWidget::AddListChanged));
 
 	ChildSlot
 	[
@@ -106,7 +35,6 @@ void SAddContentWidget::Construct(const FArguments& InArgs)
 		// Content Source Tab Page
 		+ SVerticalBox::Slot()
 		.FillHeight(3.0f)
-		.Padding(FMargin(0, 0, 0, 10))
 		[
 			SNew(SBorder)
 			.BorderImage(FAddContentDialogStyle::Get().GetBrush("AddContentDialog.TabBackground"))
@@ -158,13 +86,6 @@ void SAddContentWidget::Construct(const FArguments& InArgs)
 				]
 			]
 		]
-
-		// List of selected content
-		+ SVerticalBox::Slot()
-		.FillHeight(1.0f)
-		[
-			CreateAddListView()
-		]
 	];
 }
 
@@ -220,7 +141,7 @@ TSharedRef<SWidget> SAddContentWidget::CreateContentSourceTileView()
 	.OnGenerateTile(this, &SAddContentWidget::CreateContentSourceIconTile)
 	.OnSelectionChanged(this, &SAddContentWidget::ContentSourceTileViewSelectionChanged)
 	.ItemWidth(70)
-	.ItemHeight(100);
+	.ItemHeight(115);
 	ContentSourceTileView->SetSelection(ViewModel->GetSelectedContentSource(), ESelectInfo::Direct);
 	return ContentSourceTileView.ToSharedRef();
 }
@@ -228,14 +149,13 @@ TSharedRef<SWidget> SAddContentWidget::CreateContentSourceTileView()
 TSharedRef<ITableRow> SAddContentWidget::CreateContentSourceIconTile(TSharedPtr<FContentSourceViewModel> ContentSource, const TSharedRef<STableViewBase>& OwnerTable)
 {
 	return SNew(STableRow< TSharedPtr<FString> >, OwnerTable)
-	.Cursor(EMouseCursor::GrabHand)
-	.OnDragDetected(this, &SAddContentWidget::ContentSourceTileDragDetected)
 	[
 		SNew(SVerticalBox)
 
 		+ SVerticalBox::Slot()
 		.HAlign(EHorizontalAlignment::HAlign_Center)
 		.AutoHeight()
+		.Padding(FMargin(3))
 		[
 			SNew(SImage)
 			.Image(ContentSource->GetIconBrush().Get())
@@ -244,44 +164,62 @@ TSharedRef<ITableRow> SAddContentWidget::CreateContentSourceIconTile(TSharedPtr<
 		+ SVerticalBox::Slot()
 		.HAlign(HAlign_Center)
 		.AutoHeight()
+		.Padding(FMargin(3, 0, 3, 3))
 		[
 			SNew(STextBlock)
 			.Text(ContentSource->GetName())
 			.AutoWrapText(true)
+			.Justification(ETextJustify::Center)
 		]
 	];
 }
 
 TSharedRef<SWidget> SAddContentWidget::CreateContentSourceDetail(TSharedPtr<FContentSourceViewModel> ContentSource)
 {
-	TSharedRef<SScrollBox> Box = SNew(SScrollBox);
+	TSharedRef<SVerticalBox> VerticalBox = SNew(SVerticalBox);
 	if (ContentSource.IsValid())
 	{
-		Box->AddSlot()
-		.Padding(FMargin(0, 0, 0, 5))
-		.HAlign(EHorizontalAlignment::HAlign_Left)
+		VerticalBox->AddSlot()
 		[
-			CreateScreenshotCarousel(ContentSource)
-		];
+			SNew(SScrollBox)
+			+ SScrollBox::Slot()
+			.Padding(FMargin(0, 0, 0, 5))
+			.HAlign(EHorizontalAlignment::HAlign_Left)
+			[
+				CreateScreenshotCarousel(ContentSource)
+			]
 
-		Box->AddSlot()
-		.Padding(FMargin(0, 0, 0, 5))
-		[
-			SNew(STextBlock)
-			.TextStyle(FAddContentDialogStyle::Get(), "AddContentDialog.HeadingText")
-			.Text(ContentSource->GetName())
-			.AutoWrapText(true)
+			+SScrollBox::Slot()
+			.Padding(FMargin(0, 0, 0, 5))
+			[
+				SNew(STextBlock)
+				.TextStyle(FAddContentDialogStyle::Get(), "AddContentDialog.HeadingText")
+				.Text(ContentSource->GetName())
+				.AutoWrapText(true)
+			]
+
+			+ SScrollBox::Slot()
+			.Padding(FMargin(0, 0, 0, 5))
+			[
+				SNew(STextBlock)
+				.Text(ContentSource->GetDescription())
+				.AutoWrapText(true)
+			]
 		];
 
-		Box->AddSlot()
-		.Padding(FMargin(0, 0, 0, 5))
+		VerticalBox->AddSlot()
+		.AutoHeight()
 		[
-			SNew(STextBlock)
-			.Text(ContentSource->GetDescription())
-			.AutoWrapText(true)
+			SNew(SButton)
+			.OnClicked(this, &SAddContentWidget::AddButtonClicked)
+			.HAlign(EHorizontalAlignment::HAlign_Center)
+			[
+				SNew(STextBlock)
+				.Text(LOCTEXT("AddToProjectButton", "Add to Project"))
+			]
 		];
 	}
-	return Box;
+	return VerticalBox;
 }
 
 TSharedRef<SWidget> SAddContentWidget::CreateScreenshotCarousel(TSharedPtr<FContentSourceViewModel> ContentSource)
@@ -291,89 +229,6 @@ TSharedRef<SWidget> SAddContentWidget::CreateScreenshotCarousel(TSharedPtr<FCont
 	.WidgetItemsSource(ContentSource->GetScreenshotBrushes());
 }
 
-TSharedRef<SWidget> SAddContentWidget::CreateAddListView()
-{
-	return SNew(SVerticalBox)
-
-	+ SVerticalBox::Slot()
-	.AutoHeight()
-	.Padding(FMargin(0, 0, 0, 5))
-	[
-		SNew(STextBlock)
-		.TextStyle(FAddContentDialogStyle::Get(), "AddContentDialog.HeadingText")
-		.Text(LOCTEXT("SelectedContentToAdd", "Selected Content"))
-	]
-
-	+ SVerticalBox::Slot()
-	[
-		SNew(SOverlay)
-		
-		+ SOverlay::Slot()
-		[
-			SAssignNew(AddListView, SListView<TSharedPtr<FContentSourceViewModel>>)
-			.ListItemsSource(ViewModel->GetAddList())
-			.OnGenerateRow(this, &SAddContentWidget::CreateAddContentSourceRow)
-		]
-
-		+ SOverlay::Slot()
-		[
-			SNew(SContentSourceDragDropHandler)
-			.ViewModel(TWeakPtr<FAddContentWidgetViewModel>(ViewModel))
-			.Visibility_Lambda([this]()->EVisibility
-			{
-				if (ViewModel->GetAddList()->Num() == 0)
-				{
-					return EVisibility::Visible;
-				}
-				else
-				{
-					return FSlateApplication::Get().IsDragDropping() && FSlateApplication::Get().GetDragDroppingContent()->IsOfType<FContentSourceDragDropOp>() ?
-						EVisibility::Visible : EVisibility::Collapsed;
-				}
-			})
-			[
-				SNew(SBorder)
-				.HAlign(EHorizontalAlignment::HAlign_Center)
-				.VAlign(EVerticalAlignment::VAlign_Center)
-				[
-					SNew(STextBlock)
-					.TextStyle(FAddContentDialogStyle::Get(), "AddContentDialog.HeadingTextSmall")
-					.Text(LOCTEXT("DropContent", "+ Drop Content Here"))
-				]
-			]
-		]
-	];
-}
-
-TSharedRef<ITableRow> SAddContentWidget::CreateAddContentSourceRow(TSharedPtr<FContentSourceViewModel> ContentSource, const TSharedRef<STableViewBase>& OwnerTable)
-{
-	return SNew(STableRow< TSharedPtr<FString> >, OwnerTable)
-	.Padding(5)
-	[
-		SNew(SHorizontalBox)
-
-		+ SHorizontalBox::Slot()
-		.AutoWidth()
-		.Padding(FMargin(0, 0, 5, 0))
-		[
-			SNew(STextBlock)
-			.Text_Lambda([ContentSource]()->FText
-			{
-				return FText::Format(LOCTEXT("CategoryNameFormat", "{0} - {1}"), ContentSource->GetCategory().GetText(), ContentSource->GetName());
-			})
-		]
-
-		+ SHorizontalBox::Slot()
-		.AutoWidth()
-		[
-			SNew(SButton)
-			.ButtonStyle(FAddContentDialogStyle::Get(), "AddContentDialog.RemoveButton")
-			.ToolTipText(LOCTEXT("RemoteToolTip", "Remove"))
-			.OnClicked(this, &SAddContentWidget::RemoveAddedContentSourceClicked, ContentSource)
-		]
-	];
-}
-
 void SAddContentWidget::CategoryCheckBoxCheckStateChanged(ESlateCheckBoxState::Type CheckState, FCategoryViewModel Category)
 {
 	if (CheckState == ESlateCheckBoxState::Checked)
@@ -392,19 +247,18 @@ void SAddContentWidget::SearchTextChanged(const FText& SearchText)
 	ViewModel->SetSearchText(SearchText);
 }
 
-void SAddContentWidget::ContentSourceTileViewSelectionChanged(TSharedPtr<FContentSourceViewModel> SelectedContentSource, ESelectInfo::Type SelectInfo)
+void SAddContentWidget::ContentSourceTileViewSelectionChanged( TSharedPtr<FContentSourceViewModel> SelectedContentSource, ESelectInfo::Type SelectInfo )
 {
-	ViewModel->SetSelectedContentSource(SelectedContentSource);
+	ViewModel->SetSelectedContentSource( SelectedContentSource );
 }
 
-FReply SAddContentWidget::ContentSourceTileDragDetected(const FGeometry& MyGeometry, const FPointerEvent& MouseEvent)
+FReply SAddContentWidget::AddButtonClicked()
 {
-	TSharedPtr<FContentSourceViewModel> SelectedContentSource = ViewModel->GetSelectedContentSource();
-	if (SelectedContentSource.IsValid())
+	if (ViewModel->GetSelectedContentSource().IsValid())
 	{
-		return FReply::Handled().BeginDragDrop(FContentSourceDragDropOp::CreateShared(SelectedContentSource));
+		ViewModel->GetSelectedContentSource()->GetContentSource()->InstallToProject( "/Game" );
 	}
-	return FReply::Unhandled();
+	return FReply::Handled();
 }
 
 TSharedRef<SWidget> SAddContentWidget::CreateScreenshotWidget(TSharedPtr<FSlateBrush> ScreenshotBrush)
@@ -413,12 +267,6 @@ TSharedRef<SWidget> SAddContentWidget::CreateScreenshotWidget(TSharedPtr<FSlateB
 	.Image(ScreenshotBrush.Get());
 }
 
-FReply SAddContentWidget::RemoveAddedContentSourceClicked(TSharedPtr<FContentSourceViewModel> ContentSource)
-{
-	ViewModel->RemoveFromAddList(ContentSource);
-	return FReply::Handled();
-}
-
 void SAddContentWidget::CategoriesChanged()
 {
 	CategoryTabsContainer->SetContent(CreateCategoryTabs());
@@ -435,16 +283,6 @@ void SAddContentWidget::SelectedContentSourceChanged()
 	ContentSourceDetailContainer->SetContent(CreateContentSourceDetail(ViewModel->GetSelectedContentSource()));
 }
 
-void SAddContentWidget::AddListChanged()
-{
-	AddListView->RequestListRefresh();
-	ContentSourcesToAdd.Empty();
-	for (const TSharedPtr<FContentSourceViewModel>& ContentSource : *ViewModel->GetAddList())
-	{
-		ContentSourcesToAdd.Add(ContentSource->GetContentSource());
-	}
-}
-
 SAddContentWidget::~SAddContentWidget()
 {
 }
diff --git a/Engine/Source/Editor/AddContentDialog/Private/SAddContentWidget.h b/Engine/Source/Editor/AddContentDialog/Private/SAddContentWidget.h
--- a/Engine/Source/Editor/AddContentDialog/Private/SAddContentWidget.h
+++ b/Engine/Source/Editor/AddContentDialog/Private/SAddContentWidget.h
@@ -43,13 +43,6 @@ private:
 	/** Creates the widget the displays a screenshot in the screenshot carousel. */
 	TSharedRef<SWidget> CreateScreenshotWidget(TSharedPtr<FSlateBrush> ScreenshotBrush);
 
-	/** Creates a list view representing content which the user has selected for addition to the project. */
-	TSharedRef<SWidget> CreateAddListView();
-	
-	/** Creates a row for the list view representing content selected for addition to the project */
-	TSharedRef<ITableRow> CreateAddContentSourceRow(TSharedPtr<FContentSourceViewModel> ContentSource, const TSharedRef<STableViewBase>& OwnerTable);
-
-
 
 	/** Handles the user clicking on one of the check boxes representing the category tabs. */
 	void CategoryCheckBoxCheckStateChanged(ESlateCheckBoxState::Type CheckState, FCategoryViewModel Category);
@@ -63,12 +56,8 @@ private:
 	/** Handles the selection in the content source tile view changing */
 	void ContentSourceTileViewSelectionChanged(TSharedPtr<FContentSourceViewModel> SelectedContentSource, ESelectInfo::Type SelectInfo);
 
-	/** Handles an item from the content source tile view being dragged. */
-	FReply ContentSourceTileDragDetected(const FGeometry& MyGeometry, const FPointerEvent& MouseEvent);
-
-	/** Handles the remove button being pressed on a row in the list of content selected for addition to the project. */
-	FReply RemoveAddedContentSourceClicked(TSharedPtr<FContentSourceViewModel> ContentSource);
-
+	/** Handles the add content to project button being clicked. */
+	FReply AddButtonClicked();
 
 
 	/** Handles the available categories changing on the view model. */
@@ -80,9 +69,6 @@ private:
 	/** Handles the selected content source changing on the view model. */
 	void SelectedContentSourceChanged();
 
-	/** Handles the list of items selected for addition to the project changing on the view model. */
-	void AddListChanged();
-
 private:
 	/** The view model which represents the current data of the UI. */
 	TSharedPtr<FAddContentWidgetViewModel> ViewModel;
@@ -96,9 +82,6 @@ private:
 	/** The placeholder widget which holds the detail view for the currently selected content source. */
 	TSharedPtr<SBox> ContentSourceDetailContainer;
 
-	/** The list view which displays the content sources currently selected for addition to the project. */
-	TSharedPtr<SListView<TSharedPtr<FContentSourceViewModel>>> AddListView;
-
 	/** The content sources which the user has selected for addition to the project. */
 	TArray<TSharedPtr<IContentSource>> ContentSourcesToAdd;
 };
\ No newline at end of file
diff --git a/Engine/Source/Editor/AddContentDialog/Private/ViewModels/AddContentWidgetViewModel.cpp b/Engine/Source/Editor/AddContentDialog/Private/ViewModels/AddContentWidgetViewModel.cpp
--- a/Engine/Source/Editor/AddContentDialog/Private/ViewModels/AddContentWidgetViewModel.cpp
+++ b/Engine/Source/Editor/AddContentDialog/Private/ViewModels/AddContentWidgetViewModel.cpp
@@ -155,29 +155,4 @@ void FAddContentWidgetViewModel::TransformContentSourceToStrings(TSharedPtr<FCon
 void FAddContentWidgetViewModel::ContentSourcesChanged()
 {
 	BuildContentSourceViewModels();
-}
-
-void FAddContentWidgetViewModel::AppendToAddList(TSharedPtr<FContentSourceViewModel> ContentSourceToAppend)
-{
-	if (AddList.Contains(ContentSourceToAppend) == false)
-	{
-		AddList.Add(ContentSourceToAppend);
-		OnAddListChanged.ExecuteIfBound();
-	}
-}
-
-void FAddContentWidgetViewModel::RemoveFromAddList(TSharedPtr<FContentSourceViewModel> ContentSourceToRemove)
-{
-	AddList.Remove(ContentSourceToRemove);
-	OnAddListChanged.ExecuteIfBound();
-}
-
-const TArray<TSharedPtr<FContentSourceViewModel>>* FAddContentWidgetViewModel::GetAddList()
-{
-	return &AddList;
-}
-
-void FAddContentWidgetViewModel::SetOnAddListChanged(FOnAddListChanged OnAddListChangedIn)
-{
-	OnAddListChanged = OnAddListChangedIn;
 }
\ No newline at end of file
diff --git a/Engine/Source/Editor/AddContentDialog/Private/ViewModels/AddContentWidgetViewModel.h b/Engine/Source/Editor/AddContentDialog/Private/ViewModels/AddContentWidgetViewModel.h
--- a/Engine/Source/Editor/AddContentDialog/Private/ViewModels/AddContentWidgetViewModel.h
+++ b/Engine/Source/Editor/AddContentDialog/Private/ViewModels/AddContentWidgetViewModel.h
@@ -9,7 +9,6 @@ public:
 	DECLARE_DELEGATE(FOnCategoriesChanged)
 	DECLARE_DELEGATE(FOnContentSourcesChanged);
 	DECLARE_DELEGATE(FOnSelectedContentSourceChanged);
-	DECLARE_DELEGATE(FOnAddListChanged);
 
 	typedef TTextFilter<TSharedPtr<FContentSourceViewModel>> ContentSourceTextFilter;
 
@@ -50,18 +49,6 @@ public:
 	/** Sets the delegate which should be executed when the selected content source changes. */
 	void SetOnSelectedContentSourceChanged(FOnSelectedContentSourceChanged OnSelectedContentSourceChangedIn);
 
-	/** Gets the array of content sources which the user has selected for addition to the project. */
-	const TArray<TSharedPtr<FContentSourceViewModel>>* GetAddList();
-
-	/** Adds a content source to the array of content sources which the user has selected for addition to the project. */
-	void AppendToAddList(TSharedPtr<FContentSourceViewModel> ContentSourceToAppend);
-
-	/** Removes a content source from the array of content sources which the user has selected for addition to the project. */
-	void RemoveFromAddList(TSharedPtr<FContentSourceViewModel> ContentSourceToRemove);
-
-	/** Sets the delegate which should be executed when the contents of the array returned by GetAddList changes. */
-	void SetOnAddListChanged(FOnAddListChanged OnAddListChangedIn);
-
 private:
 	FAddContentWidgetViewModel();
 	void Initialize();
@@ -91,9 +78,6 @@ private:
 	/** A filtered array of content sources based on the currently selected category and the search text. */
 	TArray<TSharedPtr<FContentSourceViewModel>> FilteredContentSourceViewModels;
 
-	/** An array of the content sources which the user has selected for addition to the project. */
-	TArray<TSharedPtr<FContentSourceViewModel>> AddList;
-
 	/** A map which keeps track of the currently selected content source for each category. */
 	TMap<FCategoryViewModel, TSharedPtr<FContentSourceViewModel>> CategoryToSelectedContentSourceMap;
 
@@ -112,10 +96,6 @@ private:
 	/** The delegate which is executed when the currently selected content source changes. */
 	FOnSelectedContentSourceChanged OnSelectedContentSourceChanged;
 
-	/** The delegate which is executed when the contents of the array of content souces to add
-		changes. */
-	FOnAddListChanged OnAddListChanged;
-
 	/** The filter which is used to filter the content sources based on the search text. */
 	TSharedPtr<ContentSourceTextFilter> ContentSourceFilter;
 };
\ No newline at end of file
diff --git a/Engine/Source/Editor/AddContentDialog/Private/ViewModels/ContentSourceViewModel.cpp b/Engine/Source/Editor/AddContentDialog/Private/ViewModels/ContentSourceViewModel.cpp
--- a/Engine/Source/Editor/AddContentDialog/Private/ViewModels/ContentSourceViewModel.cpp
+++ b/Engine/Source/Editor/AddContentDialog/Private/ViewModels/ContentSourceViewModel.cpp
@@ -9,6 +9,8 @@
 
 const FString DefaultLanguageCode = "en";
 
+uint32 FContentSourceViewModel::ImageID = 0;
+
 FContentSourceViewModel::FContentSourceViewModel(TSharedPtr<IContentSource> ContentSourceIn)
 {
 	ContentSource = ContentSourceIn;
@@ -59,16 +61,16 @@ TArray<TSharedPtr<FSlateBrush>>* FContentSourceViewModel::GetScreenshotBrushes()
 void FContentSourceViewModel::SetupBrushes()
 {
 	FString IconBrushName = GetName().ToString() + "_" + ContentSource->GetIconData()->GetName();
-	IconBrush = CreateBrushFromRawData(FName(*IconBrushName), *ContentSource->GetIconData()->GetData());
+	IconBrush = CreateBrushFromRawData(IconBrushName, *ContentSource->GetIconData()->GetData());
 
 	for (TSharedPtr<FImageData> ScreenshotData : ContentSource->GetScreenshotData())
 	{
 		FString ScreenshotBrushName = GetName().ToString() + "_" + ScreenshotData->GetName();
-		ScreenshotBrushes.Add(CreateBrushFromRawData(FName(*ScreenshotBrushName), *ScreenshotData->GetData()));
+		ScreenshotBrushes.Add(CreateBrushFromRawData(ScreenshotBrushName, *ScreenshotData->GetData()));
 	}
 }
 
-TSharedPtr<FSlateDynamicImageBrush> FContentSourceViewModel::CreateBrushFromRawData(FName ResourceName, const TArray< uint8 >& RawData) const
+TSharedPtr<FSlateDynamicImageBrush> FContentSourceViewModel::CreateBrushFromRawData(FString ResourceNamePrefix, const TArray< uint8 >& RawData) const
 {
 	TSharedPtr< FSlateDynamicImageBrush > Brush;
 
@@ -96,7 +98,8 @@ TSharedPtr<FSlateDynamicImageBrush> FContentSourceViewModel::CreateBrushFromRawD
 
 	if (bSucceeded)
 	{
-		Brush = FSlateDynamicImageBrush::CreateWithImageData(ResourceName, FVector2D(ImageWrapper->GetWidth(), ImageWrapper->GetHeight()), DecodedImage);
+		FString UniqueResourceName = ResourceNamePrefix + "_" + FString::FromInt(ImageID++);
+		Brush = FSlateDynamicImageBrush::CreateWithImageData(FName(*UniqueResourceName), FVector2D(ImageWrapper->GetWidth(), ImageWrapper->GetHeight()), DecodedImage);
 	}
 
 	return Brush;
diff --git a/Engine/Source/Editor/AddContentDialog/Private/ViewModels/ContentSourceViewModel.h b/Engine/Source/Editor/AddContentDialog/Private/ViewModels/ContentSourceViewModel.h
--- a/Engine/Source/Editor/AddContentDialog/Private/ViewModels/ContentSourceViewModel.h
+++ b/Engine/Source/Editor/AddContentDialog/Private/ViewModels/ContentSourceViewModel.h
@@ -31,8 +31,8 @@ private:
 	/** Sets up brushes from the images data supplied by the IContentSource. */
 	void SetupBrushes();
 
-	/** Creates a slate brush from raw binary PNG formatted image data and the supplied FName. */
-	TSharedPtr<FSlateDynamicImageBrush> CreateBrushFromRawData(FName ResourceName, const TArray<uint8>& RawData) const;
+	/** Creates a slate brush from raw binary PNG formatted image data and the supplied prefix. */
+	TSharedPtr<FSlateDynamicImageBrush> CreateBrushFromRawData(FString ResourceNamePrefix, const TArray<uint8>& RawData) const;
 
 	/** Selects an FLocalizedText from an array which matches either the supplied language code, or the default language code. */
 	FLocalizedText ChooseLocalizedText(TArray<FLocalizedText> Choices, FString LanguageCode);
@@ -57,4 +57,9 @@ private:
 	/** The FLocalizedText representing the description of the content source, in the language which was active
 		the last time it was requested, or the default language if a translation was not available. */
 	FLocalizedText DescriptionText;
+
+	/** Keeps track of a unique increasing id which is appended to each brush name.  This avoids an issue
+		where two brushes are created with the same name, and then both brushes texture data gets deleted
+		when either brush is destructed. */
+	static uint32 ImageID;
 };
\ No newline at end of file
diff --git a/Engine/Source/Editor/ContentBrowser/Private/SContentBrowser.cpp b/Engine/Source/Editor/ContentBrowser/Private/SContentBrowser.cpp
--- a/Engine/Source/Editor/ContentBrowser/Private/SContentBrowser.cpp
+++ b/Engine/Source/Editor/ContentBrowser/Private/SContentBrowser.cpp
@@ -1364,12 +1364,17 @@ FReply SContentBrowser::OnSaveClicked()
 
 FReply SContentBrowser::OnAddContentClicked()
 {
-	IAddContentDialogModule& AddContentDialogModule = FModuleManager::LoadModuleChecked<IAddContentDialogModule>("AddContentDialog");
-	TSharedRef<SWindow> AddContentDialog = AddContentDialogModule.CreateDialogWindow();
+	if (AddContentDialog.IsValid() == false)
+	{
+		IAddContentDialogModule& AddContentDialogModule = FModuleManager::LoadModuleChecked<IAddContentDialogModule>("AddContentDialog");
+		TSharedRef<SWindow> Dialog = AddContentDialogModule.CreateDialogWindow();
+
+		FWidgetPath WidgetPath;
+		FSlateApplication::Get().GeneratePathToWidgetChecked(AsShared(), WidgetPath);
+		FSlateApplication::Get().AddWindowAsNativeChild(Dialog, WidgetPath.GetWindow());
 
-	FWidgetPath WidgetPath;
-	FSlateApplication::Get().GeneratePathToWidgetChecked(AsShared(), WidgetPath);
-	FSlateApplication::Get().AddWindowAsNativeChild(AddContentDialog, WidgetPath.GetWindow());
+		AddContentDialog = TWeakPtr<SWindow>(Dialog);
+	}
 
 	return FReply::Handled();
 }
diff --git a/Engine/Source/Editor/ContentBrowser/Private/SContentBrowser.h b/Engine/Source/Editor/ContentBrowser/Private/SContentBrowser.h
--- a/Engine/Source/Editor/ContentBrowser/Private/SContentBrowser.h
+++ b/Engine/Source/Editor/ContentBrowser/Private/SContentBrowser.h
@@ -167,6 +167,7 @@ private:
 	/** Saves dirty content. */
 	FReply OnSaveClicked();
 
+	/** Opens the add content dialog. */
 	FReply OnAddContentClicked();
 
 	/** Handler for when the selection set in the asset view has changed. */
@@ -375,6 +376,9 @@ private:
 	/** The splitter between the path & collection view */
 	TSharedPtr<SSplitter> PathCollectionSplitterPtr;
 
+	/** A weak pointer to the currently active add content dialog. */
+	TWeakPtr<SWindow> AddContentDialog;
+
 public: 
 
 	/** The section of EditorUserSettings in which to save content browser settings */
