Fix deprecation warnings

diff --git a/Engine/Source/Editor/AddContentDialog/Private/SAddContentWidget.cpp b/Engine/Source/Editor/AddContentDialog/Private/SAddContentWidget.cpp
--- a/Engine/Source/Editor/AddContentDialog/Private/SAddContentWidget.cpp
+++ b/Engine/Source/Editor/AddContentDialog/Private/SAddContentWidget.cpp
@@ -231,7 +231,7 @@ TSharedRef<SWidget> SAddContentWidget::CreateScreenshotCarousel(TSharedPtr<FCont
 
 void SAddContentWidget::CategoryCheckBoxCheckStateChanged(ESlateCheckBoxState::Type CheckState, FCategoryViewModel Category)
 {
-	if (CheckState == ESlateCheckBoxState::Checked)
+	if (CheckState == ECheckBoxState::Checked)
 	{
 		ViewModel->SetSelectedCategory(Category);
 	}
@@ -239,7 +239,7 @@ void SAddContentWidget::CategoryCheckBoxCheckStateChanged(ESlateCheckBoxState::T
 
 ESlateCheckBoxState::Type SAddContentWidget::GetCategoryCheckBoxCheckState(FCategoryViewModel Category) const
 {
-	return (Category == ViewModel->GetSelectedCategory()) ? ESlateCheckBoxState::Checked : ESlateCheckBoxState::Unchecked;
+	return (Category == ViewModel->GetSelectedCategory()) ? ECheckBoxState::Checked : ECheckBoxState::Unchecked;
 }
 
 void SAddContentWidget::SearchTextChanged(const FText& SearchText)
