AddContentDialog - Fix issues.
+ Fix text flicker on tab switch.
+ Don't allow multiple instances of the dialog to be opened.
+ Change the tab order to match the new project dialog.
+ Fix warnings.

diff --git a/Engine/Source/Editor/AddContentDialog/Private/AddContentDialogModule.cpp b/Engine/Source/Editor/AddContentDialog/Private/AddContentDialogModule.cpp
--- a/Engine/Source/Editor/AddContentDialog/Private/AddContentDialogModule.cpp
+++ b/Engine/Source/Editor/AddContentDialog/Private/AddContentDialogModule.cpp
@@ -32,14 +32,20 @@ public:
 		return ContentSourceProviderManager.ToSharedRef();
 	}
 
-	virtual TSharedRef<SWindow> CreateDialogWindow() override
+	virtual void ShowDialog(TSharedRef<SWindow> ParentWindow) override
 	{
-		return SNew(SAddContentDialog);
+		if (AddContentDialog.IsValid() == false)
+		{
+			TSharedRef<SWindow> Dialog = SNew(SAddContentDialog);
+			FSlateApplication::Get().AddWindowAsNativeChild(Dialog, ParentWindow);
+			AddContentDialog = TWeakPtr<SWindow>(Dialog);
+		}
+		
 	}
 
 private:
 	TSharedPtr<FContentSourceProviderManager> ContentSourceProviderManager;
-
+	TWeakPtr<SWindow> AddContentDialog;
 };
 
 IMPLEMENT_MODULE(FAddContentDialogModule, AddContentDialog);
diff --git a/Engine/Source/Editor/AddContentDialog/Private/SAddContentWidget.cpp b/Engine/Source/Editor/AddContentDialog/Private/SAddContentWidget.cpp
--- a/Engine/Source/Editor/AddContentDialog/Private/SAddContentWidget.cpp
+++ b/Engine/Source/Editor/AddContentDialog/Private/SAddContentWidget.cpp
@@ -168,7 +168,7 @@ TSharedRef<ITableRow> SAddContentWidget::CreateContentSourceIconTile(TSharedPtr<
 		[
 			SNew(STextBlock)
 			.Text(ContentSource->GetName())
-			.AutoWrapText(true)
+			.WrapTextAt(64)
 			.Justification(ETextJustify::Center)
 		]
 	];
@@ -229,7 +229,7 @@ TSharedRef<SWidget> SAddContentWidget::CreateScreenshotCarousel(TSharedPtr<FCont
 	.WidgetItemsSource(ContentSource->GetScreenshotBrushes());
 }
 
-void SAddContentWidget::CategoryCheckBoxCheckStateChanged(ESlateCheckBoxState::Type CheckState, FCategoryViewModel Category)
+void SAddContentWidget::CategoryCheckBoxCheckStateChanged(ECheckBoxState CheckState, FCategoryViewModel Category)
 {
 	if (CheckState == ECheckBoxState::Checked)
 	{
@@ -237,7 +237,7 @@ void SAddContentWidget::CategoryCheckBoxCheckStateChanged(ESlateCheckBoxState::T
 	}
 }
 
-ESlateCheckBoxState::Type SAddContentWidget::GetCategoryCheckBoxCheckState(FCategoryViewModel Category) const
+ECheckBoxState SAddContentWidget::GetCategoryCheckBoxCheckState(FCategoryViewModel Category) const
 {
 	return (Category == ViewModel->GetSelectedCategory()) ? ECheckBoxState::Checked : ECheckBoxState::Unchecked;
 }
diff --git a/Engine/Source/Editor/AddContentDialog/Private/SAddContentWidget.h b/Engine/Source/Editor/AddContentDialog/Private/SAddContentWidget.h
--- a/Engine/Source/Editor/AddContentDialog/Private/SAddContentWidget.h
+++ b/Engine/Source/Editor/AddContentDialog/Private/SAddContentWidget.h
@@ -45,10 +45,10 @@ private:
 
 
 	/** Handles the user clicking on one of the check boxes representing the category tabs. */
-	void CategoryCheckBoxCheckStateChanged(ESlateCheckBoxState::Type CheckState, FCategoryViewModel Category);
+	void CategoryCheckBoxCheckStateChanged(ECheckBoxState CheckState, FCategoryViewModel Category);
 
 	/** Gets the check state for one of the check boxes representing the category tabs. */
-	ESlateCheckBoxState::Type GetCategoryCheckBoxCheckState(FCategoryViewModel Category) const;
+	ECheckBoxState GetCategoryCheckBoxCheckState(FCategoryViewModel Category) const;
 
 	/** Handles the text in the search box changing. */
 	void SearchTextChanged(const FText& SearchText);
diff --git a/Engine/Source/Editor/AddContentDialog/Private/ViewModels/AddContentWidgetViewModel.cpp b/Engine/Source/Editor/AddContentDialog/Private/ViewModels/AddContentWidgetViewModel.cpp
--- a/Engine/Source/Editor/AddContentDialog/Private/ViewModels/AddContentWidgetViewModel.cpp
+++ b/Engine/Source/Editor/AddContentDialog/Private/ViewModels/AddContentWidgetViewModel.cpp
@@ -122,6 +122,7 @@ void FAddContentWidgetViewModel::BuildContentSourceViewModels()
 
 	if (Categories.Num() > 0)
 	{
+		Categories.Sort();
 		SelectedCategory = Categories[0];
 	}
 	OnCategoriesChanged.ExecuteIfBound();
diff --git a/Engine/Source/Editor/AddContentDialog/Private/ViewModels/CategoryViewModel.cpp b/Engine/Source/Editor/AddContentDialog/Private/ViewModels/CategoryViewModel.cpp
--- a/Engine/Source/Editor/AddContentDialog/Private/ViewModels/CategoryViewModel.cpp
+++ b/Engine/Source/Editor/AddContentDialog/Private/ViewModels/CategoryViewModel.cpp
@@ -38,18 +38,22 @@ void FCategoryViewModel::Initialize()
 	case EContentSourceCategory::BlueprintFeature:
 		Text = LOCTEXT("BlueprintFeature", "Blueprint Feature");
 		IconBrush = FAddContentDialogStyle::Get().GetBrush("AddContentDialog.BlueprintFeatureCategory");
+		SortID = 0;
 		break;
 	case EContentSourceCategory::CodeFeature:
 		Text = LOCTEXT("CodeFeature", "C++ Feature");
 		IconBrush = FAddContentDialogStyle::Get().GetBrush("AddContentDialog.CodeFeatureCategory");
+		SortID = 1;
 		break;
 	case EContentSourceCategory::Content:
 		Text = LOCTEXT("ContentPacks", "Content Packs");
 		IconBrush = FAddContentDialogStyle::Get().GetBrush("AddContentDialog.ContentPacks");
+		SortID = 2;
 		break;
 	default:
 		Text = LOCTEXT("Miscellaneous", "Miscellaneous");
 		IconBrush = FAddContentDialogStyle::Get().GetBrush("AddContentDialog.DefaultCategory");
+		SortID = 3;
 		break;
 	}
 }
diff --git a/Engine/Source/Editor/AddContentDialog/Private/ViewModels/CategoryViewModel.h b/Engine/Source/Editor/AddContentDialog/Private/ViewModels/CategoryViewModel.h
--- a/Engine/Source/Editor/AddContentDialog/Private/ViewModels/CategoryViewModel.h
+++ b/Engine/Source/Editor/AddContentDialog/Private/ViewModels/CategoryViewModel.h
@@ -21,6 +21,11 @@ public:
 		return Other.Text.EqualTo(Text) && (Other.IconBrush == IconBrush);
 	}
 
+	bool operator<(FCategoryViewModel const& Other) const
+	{
+		return SortID < Other.SortID;
+	}
+
 	uint32 GetTypeHash() const;
 
 private:
@@ -30,6 +35,7 @@ private:
 	EContentSourceCategory Category;
 	FText Text;
 	const FSlateBrush* IconBrush;
+	int SortID;
 };
 
 inline uint32 GetTypeHash(const FCategoryViewModel& CategoryViewModel)
diff --git a/Engine/Source/Editor/AddContentDialog/Public/IAddContentDialogModule.h b/Engine/Source/Editor/AddContentDialog/Public/IAddContentDialogModule.h
--- a/Engine/Source/Editor/AddContentDialog/Public/IAddContentDialogModule.h
+++ b/Engine/Source/Editor/AddContentDialog/Public/IAddContentDialogModule.h
@@ -15,5 +15,5 @@ public:
 	virtual TSharedRef<FContentSourceProviderManager> GetContentSourceProviderManager() = 0;
 
 	/** Creates a dialog for adding existing content to a project. */
-	virtual TSharedRef<SWindow> CreateDialogWindow() = 0;
+	virtual void ShowDialog(TSharedRef<SWindow> ParentWindow) = 0;
 };
\ No newline at end of file
diff --git a/Engine/Source/Editor/ContentBrowser/Private/SContentBrowser.cpp b/Engine/Source/Editor/ContentBrowser/Private/SContentBrowser.cpp
--- a/Engine/Source/Editor/ContentBrowser/Private/SContentBrowser.cpp
+++ b/Engine/Source/Editor/ContentBrowser/Private/SContentBrowser.cpp
@@ -1364,17 +1364,10 @@ FReply SContentBrowser::OnSaveClicked()
 
 FReply SContentBrowser::OnAddContentClicked()
 {
-	if (AddContentDialog.IsValid() == false)
-	{
-		IAddContentDialogModule& AddContentDialogModule = FModuleManager::LoadModuleChecked<IAddContentDialogModule>("AddContentDialog");
-		TSharedRef<SWindow> Dialog = AddContentDialogModule.CreateDialogWindow();
-
-		FWidgetPath WidgetPath;
-		FSlateApplication::Get().GeneratePathToWidgetChecked(AsShared(), WidgetPath);
-		FSlateApplication::Get().AddWindowAsNativeChild(Dialog, WidgetPath.GetWindow());
-
-		AddContentDialog = TWeakPtr<SWindow>(Dialog);
-	}
+	IAddContentDialogModule& AddContentDialogModule = FModuleManager::LoadModuleChecked<IAddContentDialogModule>("AddContentDialog");
+	FWidgetPath WidgetPath;
+	FSlateApplication::Get().GeneratePathToWidgetChecked(AsShared(), WidgetPath);
+	AddContentDialogModule.ShowDialog(WidgetPath.GetWindow());
 
 	return FReply::Handled();
 }
diff --git a/Engine/Source/Editor/ContentBrowser/Private/SContentBrowser.h b/Engine/Source/Editor/ContentBrowser/Private/SContentBrowser.h
--- a/Engine/Source/Editor/ContentBrowser/Private/SContentBrowser.h
+++ b/Engine/Source/Editor/ContentBrowser/Private/SContentBrowser.h
@@ -376,9 +376,6 @@ private:
 	/** The splitter between the path & collection view */
 	TSharedPtr<SSplitter> PathCollectionSplitterPtr;
 
-	/** A weak pointer to the currently active add content dialog. */
-	TWeakPtr<SWindow> AddContentDialog;
-
 public: 
 
 	/** The section of EditorUserSettings in which to save content browser settings */
