Refactor templates so they can be used as feature packs
-moved preview/screenshot images
-added manifests
-removed public/private subdolers from code templates
-added pack config.ini
-updated template defs to exclude feature pack related files
-added FP_FirstPerson (revised first person projects used in place of FP_FirstPerson ones - with instant hit weapon in place of projectile)
-revised source folder for packed feature packs
-added repaced packs in new folder

diff --git a/Engine/Source/Editor/AddContentDialog/Private/ContentSourceProviders/FeaturePack/FeaturePackContentSource.cpp b/Engine/Source/Editor/AddContentDialog/Private/ContentSourceProviders/FeaturePack/FeaturePackContentSource.cpp
--- a/Engine/Source/Editor/AddContentDialog/Private/ContentSourceProviders/FeaturePack/FeaturePackContentSource.cpp
+++ b/Engine/Source/Editor/AddContentDialog/Private/ContentSourceProviders/FeaturePack/FeaturePackContentSource.cpp
@@ -72,8 +72,15 @@ FFeaturePackContentSource::FFeaturePackContentSource(FString InFeaturePackPath)
 	// Load image data
 	FString IconFilename = ManifestObject->GetStringField("Thumbnail");
 	TSharedPtr<TArray<uint8>> IconImageData = MakeShareable(new TArray<uint8>());
-	LoadPakFileToBuffer(PakPlatformFile, FPaths::Combine(*MountPoint, TEXT("Media"), *IconFilename), *IconImageData);
-	IconData = MakeShareable(new FImageData(IconFilename, IconImageData));
+	FString ThumbnailFile = FPaths::Combine(*MountPoint, TEXT("Media"), *IconFilename);
+	if( LoadPakFileToBuffer(PakPlatformFile, ThumbnailFile, *IconImageData) == true)
+	{
+		IconData = MakeShareable(new FImageData(IconFilename, IconImageData));
+	}
+	else
+	{
+		UE_LOG(LogFeaturePack, Warning, TEXT("Error in Feature pack %s. Cannot find thumbnail %s."), *InFeaturePackPath, *ThumbnailFile );
+	}
 
 	const TArray<TSharedPtr<FJsonValue>> ScreenshotFilenameArray = ManifestObject->GetArrayField("Screenshots");
 	for (const TSharedPtr<FJsonValue> ScreenshotFilename : ScreenshotFilenameArray)
