Look for FeaturePacks in UE4/FeaturePacks, not UE4/FeaturePacks/Packed

diff --git a/Engine/Source/Editor/AddContentDialog/Private/ContentSourceProviders/FeaturePack/FeaturePackContentSourceProvider.cpp b/Engine/Source/Editor/AddContentDialog/Private/ContentSourceProviders/FeaturePack/FeaturePackContentSourceProvider.cpp
--- a/Engine/Source/Editor/AddContentDialog/Private/ContentSourceProviders/FeaturePack/FeaturePackContentSourceProvider.cpp
+++ b/Engine/Source/Editor/AddContentDialog/Private/ContentSourceProviders/FeaturePack/FeaturePackContentSourceProvider.cpp
@@ -29,7 +29,7 @@ public:
 
 FFeaturePackContentSourceProvider::FFeaturePackContentSourceProvider()
 {
-	FeaturePackPath = FPaths::Combine(*FPaths::RootDir(), TEXT("FeaturePacks"), TEXT("Packed"));
+	FeaturePackPath = FPaths::Combine(*FPaths::RootDir(), TEXT("FeaturePacks"));
 	StartUpDirectoryWatcher();
 	RefreshFeaturePacks();
 }
