AddContentDialog - Add an icon for the "Content Packs" tab.

diff --git a/Engine/Source/Editor/AddContentDialog/Private/AddContentDialogStyle.cpp b/Engine/Source/Editor/AddContentDialog/Private/AddContentDialogStyle.cpp
--- a/Engine/Source/Editor/AddContentDialog/Private/AddContentDialogStyle.cpp
+++ b/Engine/Source/Editor/AddContentDialog/Private/AddContentDialogStyle.cpp
@@ -82,7 +82,8 @@ TSharedRef< FSlateStyleSet > FAddContentDialogStyle::Create()
 
 	Style->Set("AddContentDialog.BlueprintFeatureCategory", new IMAGE_BRUSH("BlueprintFeature", FVector2D(32.0f, 32.0f)));
 	Style->Set("AddContentDialog.CodeFeatureCategory", new IMAGE_BRUSH("CodeFeature", FVector2D(32.0f, 32.0f)));
-	Style->Set("AddContentDialog.CategoryCategory", new FSlateNoResource());
+	Style->Set("AddContentDialog.ContentPackCategory", new IMAGE_BRUSH( "ContentPack", FVector2D(32.0f, 32.0f)));
+	Style->Set("AddContentDialog.UnknownCategory", new FSlateNoResource());
 
 	return Style;
 }
diff --git a/Engine/Source/Editor/AddContentDialog/Private/ViewModels/CategoryViewModel.cpp b/Engine/Source/Editor/AddContentDialog/Private/ViewModels/CategoryViewModel.cpp
--- a/Engine/Source/Editor/AddContentDialog/Private/ViewModels/CategoryViewModel.cpp
+++ b/Engine/Source/Editor/AddContentDialog/Private/ViewModels/CategoryViewModel.cpp
@@ -47,12 +47,12 @@ void FCategoryViewModel::Initialize()
 		break;
 	case EContentSourceCategory::Content:
 		Text = LOCTEXT("ContentPacks", "Content Packs");
-		IconBrush = FAddContentDialogStyle::Get().GetBrush("AddContentDialog.ContentPacks");
+		IconBrush = FAddContentDialogStyle::Get().GetBrush("AddContentDialog.ContentPackCategory");
 		SortID = 2;
 		break;
 	default:
 		Text = LOCTEXT("Miscellaneous", "Miscellaneous");
-		IconBrush = FAddContentDialogStyle::Get().GetBrush("AddContentDialog.DefaultCategory");
+		IconBrush = FAddContentDialogStyle::Get().GetBrush("AddContentDialog.UnknownCategory");
 		SortID = 3;
 		break;
 	}
