AddContentDialog - Automatically save imported assets from added content.

diff --git a/Engine/Source/Editor/AddContentDialog/AddContentDialog.Build.cs b/Engine/Source/Editor/AddContentDialog/AddContentDialog.Build.cs
--- a/Engine/Source/Editor/AddContentDialog/AddContentDialog.Build.cs
+++ b/Engine/Source/Editor/AddContentDialog/AddContentDialog.Build.cs
@@ -24,7 +24,8 @@ public class AddContentDialog : ModuleRules
 				"DesktopPlatform",
 				"PakFile",
 				"ImageWrapper",
-				"UnrealEd"
+				"UnrealEd",
+				"CoreUObject"
 			}
 		);
 
diff --git a/Engine/Source/Editor/AddContentDialog/Private/ContentSourceProviders/FeaturePack/FeaturePackContentSource.cpp b/Engine/Source/Editor/AddContentDialog/Private/ContentSourceProviders/FeaturePack/FeaturePackContentSource.cpp
--- a/Engine/Source/Editor/AddContentDialog/Private/ContentSourceProviders/FeaturePack/FeaturePackContentSource.cpp
+++ b/Engine/Source/Editor/AddContentDialog/Private/ContentSourceProviders/FeaturePack/FeaturePackContentSource.cpp
@@ -148,6 +148,13 @@ bool FFeaturePackContentSource::InstallToProject(FString InstallPath)
 		}
 		else
 		{
+			// Save any imported assets.
+			TArray<UPackage*> ToSave;
+			for (auto ImportedObject : ImportedObjects)
+			{
+				ToSave.AddUnique(ImportedObject->GetOutermost());
+			}
+			FEditorFileUtils::PromptForCheckoutAndSave( ToSave, /*bCheckDirty=*/ false, /*bPromptToSave=*/ false );
 			bResult = true;
 		}
 	}
