AddContentDialog - Add missing inlcude to fix the build.

diff --git a/Engine/Source/Editor/AddContentDialog/Private/ContentSourceProviders/FeaturePack/FeaturePackContentSource.cpp b/Engine/Source/Editor/AddContentDialog/Private/ContentSourceProviders/FeaturePack/FeaturePackContentSource.cpp
--- a/Engine/Source/Editor/AddContentDialog/Private/ContentSourceProviders/FeaturePack/FeaturePackContentSource.cpp
+++ b/Engine/Source/Editor/AddContentDialog/Private/ContentSourceProviders/FeaturePack/FeaturePackContentSource.cpp
@@ -5,6 +5,7 @@
 
 #include "AssetToolsModule.h"
 #include "IPlatformFilePak.h"
+#include "FileHelpers.h"
 
 DEFINE_LOG_CATEGORY_STATIC(LogFeaturePack, Log, All);
 
