Added asset/class field to manifest
Added display of asset/class info used to Add Content dialog
Fixed several cut/pasted comments regarding UMG styles that referred to shooter game

diff --git a/Engine/Source/Developer/FriendsAndChat/Public/FriendsAndChatStyle.h b/Engine/Source/Developer/FriendsAndChat/Public/FriendsAndChatStyle.h
--- a/Engine/Source/Developer/FriendsAndChat/Public/FriendsAndChatStyle.h
+++ b/Engine/Source/Developer/FriendsAndChat/Public/FriendsAndChatStyle.h
@@ -312,7 +312,7 @@ public:
 	/** reloads textures used by slate renderer */
 	static void ReloadTextures();
 
-	/** @return The Slate style set for the Shooter game */
+	/** @return The Slate style set for the Friends and chat module */
 	static const ISlateStyle& Get();
 
 	static FName GetStyleSetName();
diff --git a/Engine/Source/Developer/WidgetCarousel/Public/WidgetCarouselStyle.h b/Engine/Source/Developer/WidgetCarousel/Public/WidgetCarouselStyle.h
--- a/Engine/Source/Developer/WidgetCarousel/Public/WidgetCarouselStyle.h
+++ b/Engine/Source/Developer/WidgetCarousel/Public/WidgetCarouselStyle.h
@@ -79,7 +79,7 @@ public:
 	/** reloads textures used by slate renderer */
 	static void ReloadTextures();
 
-	/** @return The Slate style set for the Shooter game */
+	/** @return The Slate style set for the widget carousel */
 	static const ISlateStyle& Get();
 
 	static FName GetStyleSetName();
diff --git a/Engine/Source/Editor/AddContentDialog/Private/AddContentDialogStyle.h b/Engine/Source/Editor/AddContentDialog/Private/AddContentDialogStyle.h
--- a/Engine/Source/Editor/AddContentDialog/Private/AddContentDialogStyle.h
+++ b/Engine/Source/Editor/AddContentDialog/Private/AddContentDialogStyle.h
@@ -16,7 +16,7 @@ public:
 	/** reloads textures used by slate renderer */
 	static void ReloadTextures();
 
-	/** @return The Slate style set for the Shooter game */
+	/** @return The Slate style set for the Add Content Dialog */
 	static const ISlateStyle& Get();
 
 	static FName GetStyleSetName();
diff --git a/Engine/Source/Editor/AddContentDialog/Private/ContentSourceProviders/FeaturePack/FeaturePackContentSource.cpp b/Engine/Source/Editor/AddContentDialog/Private/ContentSourceProviders/FeaturePack/FeaturePackContentSource.cpp
--- a/Engine/Source/Editor/AddContentDialog/Private/ContentSourceProviders/FeaturePack/FeaturePackContentSource.cpp
+++ b/Engine/Source/Editor/AddContentDialog/Private/ContentSourceProviders/FeaturePack/FeaturePackContentSource.cpp
@@ -52,7 +52,19 @@ FFeaturePackContentSource::FFeaturePackContentSource(FString InFeaturePackPath)
 			FText::FromString(LocalizedDescriptionObject->GetStringField("Text"))));
 	}
 
-	FString CategoryString = ManifestObject->GetStringField("Category");
+	// Parse asset types field
+	for (TSharedPtr<FJsonValue> AssetTypesValue : ManifestObject->GetArrayField("AssetTypes"))
+	{
+		TSharedPtr<FJsonObject> LocalizedAssetTypesObject = AssetTypesValue->AsObject();
+		LocalizedAssetTypesList.Add(FLocalizedText(
+			LocalizedAssetTypesObject->GetStringField("Language"),
+			FText::FromString(LocalizedAssetTypesObject->GetStringField("Text"))));
+	}
+	
+	// Parse class types field
+	ClassTypes = ManifestObject->GetStringField("ClassTypes");
+	
+	FString CategoryString = ManifestObject->GetStringField("Category");	
 	if (CategoryString == "CodeFeature")
 	{
 		Category = EContentSourceCategory::CodeFeature;
@@ -115,6 +127,16 @@ TArray<FLocalizedText> FFeaturePackContentSource::GetLocalizedDescriptions()
 	return LocalizedDescriptions;
 }
 
+TArray<FLocalizedText> FFeaturePackContentSource::GetLocalizedAssetTypes()
+{
+	return LocalizedAssetTypesList;
+}
+
+FString FFeaturePackContentSource::GetClassTypesUsed()
+{
+	return ClassTypes;
+}
+
 EContentSourceCategory FFeaturePackContentSource::GetCategory()
 {
 	return Category;
diff --git a/Engine/Source/Editor/AddContentDialog/Private/ContentSourceProviders/FeaturePack/FeaturePackContentSource.h b/Engine/Source/Editor/AddContentDialog/Private/ContentSourceProviders/FeaturePack/FeaturePackContentSource.h
--- a/Engine/Source/Editor/AddContentDialog/Private/ContentSourceProviders/FeaturePack/FeaturePackContentSource.h
+++ b/Engine/Source/Editor/AddContentDialog/Private/ContentSourceProviders/FeaturePack/FeaturePackContentSource.h
@@ -14,6 +14,8 @@ public:
 	virtual TArray<FLocalizedText> GetLocalizedDescriptions() override;
 	
 	virtual EContentSourceCategory GetCategory() override;
+	virtual TArray<FLocalizedText> GetLocalizedAssetTypes() override;
+	virtual FString GetClassTypesUsed() override;
 
 	virtual TSharedPtr<FImageData> GetIconData() override;
 	virtual TArray<TSharedPtr<FImageData>> GetScreenshotData() override;
@@ -33,5 +35,7 @@ private:
 	EContentSourceCategory Category;
 	TSharedPtr<FImageData> IconData;
 	TArray<TSharedPtr<FImageData>> ScreenshotData;
+	TArray<FLocalizedText> LocalizedAssetTypesList;
+	FString ClassTypes;
 	bool bPackValid;
 };
\ No newline at end of file
diff --git a/Engine/Source/Editor/AddContentDialog/Private/IContentSource.h b/Engine/Source/Editor/AddContentDialog/Private/IContentSource.h
--- a/Engine/Source/Editor/AddContentDialog/Private/IContentSource.h
+++ b/Engine/Source/Editor/AddContentDialog/Private/IContentSource.h
@@ -94,6 +94,12 @@ public:
 	/** Gets an array or image data for screenshots for the content source. */
 	virtual TArray<TSharedPtr<FImageData>> GetScreenshotData() = 0;
 
+	/** Gets the asset types used in this pack. */
+	virtual TArray<FLocalizedText> GetLocalizedAssetTypes() = 0;
+
+	/** Gets the class types used in this pack. */
+	virtual FString GetClassTypesUsed() = 0;
+
 	/*
 	 * Installs the content in the content source to the specific path. 
 	 * @returns true if install suceeded
diff --git a/Engine/Source/Editor/AddContentDialog/Private/SAddContentWidget.cpp b/Engine/Source/Editor/AddContentDialog/Private/SAddContentWidget.cpp
--- a/Engine/Source/Editor/AddContentDialog/Private/SAddContentWidget.cpp
+++ b/Engine/Source/Editor/AddContentDialog/Private/SAddContentWidget.cpp
@@ -205,8 +205,42 @@ TSharedRef<SWidget> SAddContentWidget::CreateContentSourceDetail(TSharedPtr<FCon
 				.Text(ContentSource->GetDescription())
 				.AutoWrapText(true)
 			]
+		
+			+ SScrollBox::Slot()
+			.Padding(FMargin(0, 0, 0, 5))
+			[
+				SNew(STextBlock)						
+				.Visibility(ContentSource->GetAssetTypes().IsEmpty() == false ? EVisibility::Visible : EVisibility::Collapsed)
+				.TextStyle(FAddContentDialogStyle::Get(), "AddContentDialog.HeadingText")
+				.Text(LOCTEXT("FeaturePackAssetReferences", "Asset types used in this pack:"))
+			]
+			+ SScrollBox::Slot()
+			.Padding(FMargin(0, 0, 0, 5))
+			[
+				SNew(STextBlock)
+				.Text(ContentSource->GetAssetTypes())
+				.Visibility(ContentSource->GetAssetTypes().IsEmpty() == false ? EVisibility::Visible : EVisibility::Collapsed)				
+				.AutoWrapText(true)
+			]
+		
+			+ SScrollBox::Slot()
+			.Padding(FMargin(0, 0, 0, 5))
+			[
+				SNew(STextBlock)
+				.TextStyle(FAddContentDialogStyle::Get(), "AddContentDialog.HeadingText")
+				.Visibility(ContentSource->GetClassTypes().IsEmpty() == false ? EVisibility::Visible : EVisibility::Collapsed)
+				.Text(LOCTEXT("FeaturePackClassReferences", "Class types used in this pack:"))
+			] 
+			+ SScrollBox::Slot()
+			.Padding(FMargin(0, 0, 0, 5))
+			[
+				SNew(STextBlock)
+				.Text(ContentSource->GetClassTypes())
+				.Visibility(ContentSource->GetClassTypes().IsEmpty() == false ? EVisibility::Visible : EVisibility::Collapsed)
+				.AutoWrapText(true)
+			]
 		];
-
+	
 		VerticalBox->AddSlot()
 		.AutoHeight()
 		[
diff --git a/Engine/Source/Editor/AddContentDialog/Private/ViewModels/ContentSourceViewModel.cpp b/Engine/Source/Editor/AddContentDialog/Private/ViewModels/ContentSourceViewModel.cpp
--- a/Engine/Source/Editor/AddContentDialog/Private/ViewModels/ContentSourceViewModel.cpp
+++ b/Engine/Source/Editor/AddContentDialog/Private/ViewModels/ContentSourceViewModel.cpp
@@ -43,6 +43,22 @@ FText FContentSourceViewModel::GetDescription()
 	return DescriptionText.GetText();
 }
 
+FText FContentSourceViewModel::GetAssetTypes()
+{
+	FString CurrentLanguage = FInternationalization::Get().GetCurrentCulture()->GetTwoLetterISOLanguageName();
+	if (AssetTypeText.GetTwoLetterLanguage() != CurrentLanguage)
+	{
+		AssetTypeText = ChooseLocalizedText(ContentSource->GetLocalizedAssetTypes(), CurrentLanguage);
+	}
+	return AssetTypeText.GetText();
+
+}
+
+FString FContentSourceViewModel::GetClassTypes()
+{
+	return ContentSource->GetClassTypesUsed();
+}
+
 FCategoryViewModel FContentSourceViewModel::GetCategory()
 {
 	return Category;
diff --git a/Engine/Source/Editor/AddContentDialog/Private/ViewModels/ContentSourceViewModel.h b/Engine/Source/Editor/AddContentDialog/Private/ViewModels/ContentSourceViewModel.h
--- a/Engine/Source/Editor/AddContentDialog/Private/ViewModels/ContentSourceViewModel.h
+++ b/Engine/Source/Editor/AddContentDialog/Private/ViewModels/ContentSourceViewModel.h
@@ -18,6 +18,12 @@ public:
 	/** Gets the description of this content source. */
 	FText GetDescription();
 
+	/** Gets the asset types used in this content source. */
+	FText GetAssetTypes();
+
+	/** Gets the class types used in this content source. */
+	FString GetClassTypes();
+
 	/** Gets the view model for the category for this content source. */
 	FCategoryViewModel GetCategory();
 
@@ -58,6 +64,10 @@ private:
 		the last time it was requested, or the default language if a translation was not available. */
 	FLocalizedText DescriptionText;
 
+	/** The FLocalizedText representing the asset types used by the content source, in the language which was active
+		the last time it was requested, or the default language if a translation was not available. */
+	FLocalizedText AssetTypeText;
+
 	/** Keeps track of a unique increasing id which is appended to each brush name.  This avoids an issue
 		where two brushes are created with the same name, and then both brushes texture data gets deleted
 		when either brush is destructed. */
diff --git a/Engine/Source/Runtime/UMG/Public/UMGStyle.h b/Engine/Source/Runtime/UMG/Public/UMGStyle.h
--- a/Engine/Source/Runtime/UMG/Public/UMGStyle.h
+++ b/Engine/Source/Runtime/UMG/Public/UMGStyle.h
@@ -16,7 +16,7 @@ public:
 	/** reloads textures used by slate renderer */
 	static void ReloadTextures();
 
-	/** @return The Slate style set for the Shooter game */
+	/** @return The Slate style set for the UMG Style */
 	static const ISlateStyle& Get();
 
 	static FName GetStyleSetName();
