AddContentDialog - Disable multiselect in the tile view.

diff --git a/Engine/Source/Editor/AddContentDialog/Private/SAddContentWidget.cpp b/Engine/Source/Editor/AddContentDialog/Private/SAddContentWidget.cpp
--- a/Engine/Source/Editor/AddContentDialog/Private/SAddContentWidget.cpp
+++ b/Engine/Source/Editor/AddContentDialog/Private/SAddContentWidget.cpp
@@ -141,7 +141,8 @@ TSharedRef<SWidget> SAddContentWidget::CreateContentSourceTileView()
 	.OnGenerateTile(this, &SAddContentWidget::CreateContentSourceIconTile)
 	.OnSelectionChanged(this, &SAddContentWidget::ContentSourceTileViewSelectionChanged)
 	.ItemWidth(70)
-	.ItemHeight(115);
+	.ItemHeight(115)
+	.SelectionMode(ESelectionMode::Single);
 	ContentSourceTileView->SetSelection(ViewModel->GetSelectedContentSource(), ESelectInfo::Direct);
 	return ContentSourceTileView.ToSharedRef();
 }
