Revised  RefreshFeaturePacks ignore non-upak files

diff --git a/Engine/Source/Editor/AddContentDialog/Private/ContentSourceProviders/FeaturePack/FeaturePackContentSourceProvider.cpp b/Engine/Source/Editor/AddContentDialog/Private/ContentSourceProviders/FeaturePack/FeaturePackContentSourceProvider.cpp
--- a/Engine/Source/Editor/AddContentDialog/Private/ContentSourceProviders/FeaturePack/FeaturePackContentSourceProvider.cpp
+++ b/Engine/Source/Editor/AddContentDialog/Private/ContentSourceProviders/FeaturePack/FeaturePackContentSourceProvider.cpp
@@ -80,10 +80,13 @@ void FFeaturePackContentSourceProvider::RefreshFeaturePacks()
 	PlatformFile.IterateDirectory( *FeaturePackPath, DirectoryVisitor );
 	for ( auto FeaturePackFile : DirectoryVisitor.Files )
 	{
-		TUniquePtr<FFeaturePackContentSource> NewContentSource = MakeUnique<FFeaturePackContentSource>( FeaturePackFile );
-		if ( NewContentSource->IsDataValid() )
+		if( FeaturePackFile.EndsWith(TEXT(".upack")) == true)
 		{
-			ContentSources.Add( MakeShareable( NewContentSource.Release() ) );
+			TUniquePtr<FFeaturePackContentSource> NewContentSource = MakeUnique<FFeaturePackContentSource>(FeaturePackFile);
+			if (NewContentSource->IsDataValid())
+			{
+				ContentSources.Add(MakeShareable(NewContentSource.Release()));
+			}
 		}
 	}
 	OnContentSourcesChanged.ExecuteIfBound();
