Added some valid checks in FContentSourceViewModel::SetupBrushes

diff --git a/Engine/Source/Editor/AddContentDialog/Private/ViewModels/ContentSourceViewModel.cpp b/Engine/Source/Editor/AddContentDialog/Private/ViewModels/ContentSourceViewModel.cpp
--- a/Engine/Source/Editor/AddContentDialog/Private/ViewModels/ContentSourceViewModel.cpp
+++ b/Engine/Source/Editor/AddContentDialog/Private/ViewModels/ContentSourceViewModel.cpp
@@ -76,14 +76,20 @@ TArray<TSharedPtr<FSlateBrush>>* FContentSourceViewModel::GetScreenshotBrushes()
 
 void FContentSourceViewModel::SetupBrushes()
 {
-	FString IconBrushName = GetName().ToString() + "_" + ContentSource->GetIconData()->GetName();
-	IconBrush = CreateBrushFromRawData(IconBrushName, *ContentSource->GetIconData()->GetData());
+	if( ContentSource->GetIconData().IsValid() )
+	{
+		FString IconBrushName = GetName().ToString() + "_" + ContentSource->GetIconData()->GetName();
+		IconBrush = CreateBrushFromRawData(IconBrushName, *ContentSource->GetIconData()->GetData());
+	}
 
 	for (TSharedPtr<FImageData> ScreenshotData : ContentSource->GetScreenshotData())
 	{
-		FString ScreenshotBrushName = GetName().ToString() + "_" + ScreenshotData->GetName();
-		ScreenshotBrushes.Add(CreateBrushFromRawData(ScreenshotBrushName, *ScreenshotData->GetData()));
-	}
+		if( ScreenshotData.IsValid() == true )
+		{
+			FString ScreenshotBrushName = GetName().ToString() + "_" + ScreenshotData->GetName();
+			ScreenshotBrushes.Add(CreateBrushFromRawData(ScreenshotBrushName, *ScreenshotData->GetData()));
+		}
+	}	
 }
 
 TSharedPtr<FSlateDynamicImageBrush> FContentSourceViewModel::CreateBrushFromRawData(FString ResourceNamePrefix, const TArray< uint8 >& RawData) const
