Minor stability revisions to feature pack related code.

diff --git a/Engine/Source/Editor/AddContentDialog/Private/ViewModels/ContentSourceViewModel.cpp b/Engine/Source/Editor/AddContentDialog/Private/ViewModels/ContentSourceViewModel.cpp
--- a/Engine/Source/Editor/AddContentDialog/Private/ViewModels/ContentSourceViewModel.cpp
+++ b/Engine/Source/Editor/AddContentDialog/Private/ViewModels/ContentSourceViewModel.cpp
@@ -104,7 +104,7 @@ TSharedPtr<FSlateDynamicImageBrush> FContentSourceViewModel::CreateBrushFromRawD
 	TArray<uint8> DecodedImage;
 	IImageWrapperModule& ImageWrapperModule = FModuleManager::LoadModuleChecked<IImageWrapperModule>(FName("ImageWrapper"));
 	IImageWrapperPtr ImageWrapper = ImageWrapperModule.CreateImageWrapper(EImageFormat::PNG);
-	if (ImageWrapper.IsValid() && ImageWrapper->SetCompressed(RawData.GetData(), RawData.Num()))
+	if (ImageWrapper.IsValid() && (RawData.Num() > 0) && ImageWrapper->SetCompressed(RawData.GetData(), RawData.Num()))
 	{
 		Width = ImageWrapper->GetWidth();
 		Height = ImageWrapper->GetHeight();
