AddContentDialog - Clear the detail view when the user clears the tile view selection.

diff --git a/Engine/Source/Editor/AddContentDialog/Private/ViewModels/AddContentWidgetViewModel.cpp b/Engine/Source/Editor/AddContentDialog/Private/ViewModels/AddContentWidgetViewModel.cpp
--- a/Engine/Source/Editor/AddContentDialog/Private/ViewModels/AddContentWidgetViewModel.cpp
+++ b/Engine/Source/Editor/AddContentDialog/Private/ViewModels/AddContentWidgetViewModel.cpp
@@ -50,7 +50,7 @@ FCategoryViewModel FAddContentWidgetViewModel::GetSelectedCategory()
 void FAddContentWidgetViewModel::SetSelectedCategory(FCategoryViewModel SelectedCategoryIn)
 {
 	SelectedCategory = SelectedCategoryIn;
-	FilterContentSourceViewModels();
+	UpdateFilteredContentSourcesAndSelection(true);
 	OnSelectedContentSourceChanged.ExecuteIfBound();
 }
 
@@ -58,7 +58,7 @@ void FAddContentWidgetViewModel::SetSearchText(FText SearchTextIn)
 {
 	SearchText = SearchTextIn;
 	ContentSourceFilter->SetRawFilterText(SearchTextIn);
-	FilterContentSourceViewModels();
+	UpdateFilteredContentSourcesAndSelection(true);
 }
 
 const TArray<TSharedPtr<FContentSourceViewModel>>* FAddContentWidgetViewModel::GetContentSources()
@@ -83,16 +83,12 @@ TSharedPtr<FContentSourceViewModel> FAddContentWidgetViewModel::GetSelectedConte
 
 void FAddContentWidgetViewModel::SetSelectedContentSource(TSharedPtr<FContentSourceViewModel> SelectedContentSourceIn)
 {
-	// Ignore null selections.
-	if (SelectedContentSourceIn.IsValid())
+	// Ignore selecting the currently selected item.
+	TSharedPtr<FContentSourceViewModel> SelectedContentSource = GetSelectedContentSource();
+	if (SelectedContentSource != SelectedContentSourceIn)
 	{
-		// Ignore selecting the currently selected item.
-		TSharedPtr<FContentSourceViewModel> SelectedContentSource = GetSelectedContentSource();
-		if (SelectedContentSource != SelectedContentSourceIn)
-		{
-			CategoryToSelectedContentSourceMap.Add(SelectedCategory) = SelectedContentSourceIn;
-			OnSelectedContentSourceChanged.ExecuteIfBound();
-		}
+		CategoryToSelectedContentSourceMap.Add(SelectedCategory) = SelectedContentSourceIn;
+		OnSelectedContentSourceChanged.ExecuteIfBound();
 	}
 }
 
@@ -106,6 +102,7 @@ void FAddContentWidgetViewModel::BuildContentSourceViewModels()
 	Categories.Empty();
 	ContentSourceViewModels.Empty();
 	FilteredContentSourceViewModels.Empty();
+	CategoryToSelectedContentSourceMap.Empty();
 	
 	for (auto ContentSourceProvider : ContentSourceProviders)
 	{
@@ -123,13 +120,18 @@ void FAddContentWidgetViewModel::BuildContentSourceViewModels()
 	if (Categories.Num() > 0)
 	{
 		Categories.Sort();
-		SelectedCategory = Categories[0];
+		// Update the current selection for all categories.  Do this in reverse order so that the first category
+		// remains selected when finished.
+		for (int i = Categories.Num() - 1; i >= 0; i--)
+		{
+			SelectedCategory = Categories[i];
+			UpdateFilteredContentSourcesAndSelection(false);
+		}
 	}
 	OnCategoriesChanged.ExecuteIfBound();
-	FilterContentSourceViewModels();
 }
 
-void FAddContentWidgetViewModel::FilterContentSourceViewModels()
+void FAddContentWidgetViewModel::UpdateFilteredContentSourcesAndSelection(bool bAllowEmptySelection)
 {
 	FilteredContentSourceViewModels.Empty();
 	for (auto ContentSource : ContentSourceViewModels)
@@ -143,7 +145,11 @@ void FAddContentWidgetViewModel::FilterContentSourceViewModels()
 
 	if (FilteredContentSourceViewModels.Contains(GetSelectedContentSource()) == false)
 	{
-		TSharedPtr<FContentSourceViewModel> NewSelectedContentSource = FilteredContentSourceViewModels.Num() > 0 ? FilteredContentSourceViewModels[0] : TSharedPtr<FContentSourceViewModel>();
+		TSharedPtr<FContentSourceViewModel> NewSelectedContentSource;
+		if (bAllowEmptySelection == false && FilteredContentSourceViewModels.Num() > 0)
+		{
+			NewSelectedContentSource =  FilteredContentSourceViewModels[0];
+		}
 		SetSelectedContentSource(NewSelectedContentSource);
 	}
 }
diff --git a/Engine/Source/Editor/AddContentDialog/Private/ViewModels/AddContentWidgetViewModel.h b/Engine/Source/Editor/AddContentDialog/Private/ViewModels/AddContentWidgetViewModel.h
--- a/Engine/Source/Editor/AddContentDialog/Private/ViewModels/AddContentWidgetViewModel.h
+++ b/Engine/Source/Editor/AddContentDialog/Private/ViewModels/AddContentWidgetViewModel.h
@@ -56,8 +56,9 @@ private:
 	/** Builds view models for all available content sources. */
 	void BuildContentSourceViewModels();
 
-	/** Filters the current set of content sources based on the selected category and the search text. */
-	void FilterContentSourceViewModels();
+	/** Filters the current set of content sources based on the selected category and the search text, and updates the selection.
+	  * @param bAllowEmptySelection Determines whether or not to select the first available item after filtering if the selected item is null, or not in the filtered set. */
+	void UpdateFilteredContentSourcesAndSelection(bool bAllowEmptySelection);
 
 	/** Converts a content source item to an array of strings for processing by the TTextFilter. */
 	void TransformContentSourceToStrings(TSharedPtr<FContentSourceViewModel> Item, OUT TArray< FString >& Array);
