Added sortkey to feature pack manifest

diff --git a/Engine/Source/Editor/AddContentDialog/Private/ContentSourceProviders/FeaturePack/FeaturePackContentSource.cpp b/Engine/Source/Editor/AddContentDialog/Private/ContentSourceProviders/FeaturePack/FeaturePackContentSource.cpp
--- a/Engine/Source/Editor/AddContentDialog/Private/ContentSourceProviders/FeaturePack/FeaturePackContentSource.cpp
+++ b/Engine/Source/Editor/AddContentDialog/Private/ContentSourceProviders/FeaturePack/FeaturePackContentSource.cpp
@@ -68,6 +68,10 @@ FFeaturePackContentSource::FFeaturePackContentSource(FString InFeaturePackPath)
 	// Parse initial focus asset
 	FocusAssetIdent = ManifestObject->GetStringField("FocusAsset");
 
+	// Use the path as the sort key - it will be alphabetical that way
+	SortKey = FeaturePackPath;
+	ManifestObject->TryGetStringField("SortKey", SortKey);
+
 	FString CategoryString = ManifestObject->GetStringField("Category");	
 	if (CategoryString == "CodeFeature")
 	{
@@ -220,3 +224,8 @@ FString FFeaturePackContentSource::GetFocusAssetName() const
 {
 	return FocusAssetIdent;
 }
+
+FString FFeaturePackContentSource::GetSortKey() const
+{
+	return SortKey;
+}
\ No newline at end of file
diff --git a/Engine/Source/Editor/AddContentDialog/Private/ContentSourceProviders/FeaturePack/FeaturePackContentSource.h b/Engine/Source/Editor/AddContentDialog/Private/ContentSourceProviders/FeaturePack/FeaturePackContentSource.h
--- a/Engine/Source/Editor/AddContentDialog/Private/ContentSourceProviders/FeaturePack/FeaturePackContentSource.h
+++ b/Engine/Source/Editor/AddContentDialog/Private/ContentSourceProviders/FeaturePack/FeaturePackContentSource.h
@@ -15,6 +15,7 @@ public:
 	
 	virtual EContentSourceCategory GetCategory() override;
 	virtual TArray<FLocalizedText> GetLocalizedAssetTypes() override;
+	virtual FString GetSortKey() const override;
 	virtual FString GetClassTypesUsed() override;
 
 	virtual TSharedPtr<FImageData> GetIconData() override;
@@ -40,4 +41,5 @@ private:
 	FString ClassTypes;
 	bool bPackValid;
 	FString FocusAssetIdent;
+	FString SortKey;
 };
\ No newline at end of file
diff --git a/Engine/Source/Editor/AddContentDialog/Private/ContentSourceProviders/FeaturePack/FeaturePackContentSourceProvider.cpp b/Engine/Source/Editor/AddContentDialog/Private/ContentSourceProviders/FeaturePack/FeaturePackContentSourceProvider.cpp
--- a/Engine/Source/Editor/AddContentDialog/Private/ContentSourceProviders/FeaturePack/FeaturePackContentSourceProvider.cpp
+++ b/Engine/Source/Editor/AddContentDialog/Private/ContentSourceProviders/FeaturePack/FeaturePackContentSourceProvider.cpp
@@ -72,6 +72,12 @@ void FFeaturePackContentSourceProvider::OnFeaturePackDirectoryChanged( const TAr
 	RefreshFeaturePacks();
 }
 
+/** Sorting function sort keys. */
+struct FFeaturePackCompareSortKey
+{
+	FORCEINLINE bool operator()( TSharedRef<IContentSource> const& A, TSharedRef<IContentSource> const& B) const { return A->GetSortKey() < B->GetSortKey(); }
+};
+
 void FFeaturePackContentSourceProvider::RefreshFeaturePacks()
 {
 	ContentSources.Empty();
@@ -89,6 +95,7 @@ void FFeaturePackContentSourceProvider::RefreshFeaturePacks()
 			}
 		}
 	}
+	ContentSources.Sort(FFeaturePackCompareSortKey());
 	OnContentSourcesChanged.ExecuteIfBound();
 }
 
diff --git a/Engine/Source/Editor/AddContentDialog/Private/IContentSource.h b/Engine/Source/Editor/AddContentDialog/Private/IContentSource.h
--- a/Engine/Source/Editor/AddContentDialog/Private/IContentSource.h
+++ b/Engine/Source/Editor/AddContentDialog/Private/IContentSource.h
@@ -100,6 +100,8 @@ public:
 	/** Gets the class types used in this pack. */
 	virtual FString GetClassTypesUsed() = 0;
 
+	/** Gets the category for the content source. */
+	virtual FString GetSortKey() const { return FString(TEXT(""));};
 	/*
 	 * Installs the content in the content source to the specific path. 
 	 * @returns true if install suceeded
