AddContentDialog - Update the style and layout of the "Add to Project" button.

Jira: UE-10273

diff --git a/Engine/Source/Editor/AddContentDialog/Private/AddContentDialogStyle.cpp b/Engine/Source/Editor/AddContentDialog/Private/AddContentDialogStyle.cpp
--- a/Engine/Source/Editor/AddContentDialog/Private/AddContentDialogStyle.cpp
+++ b/Engine/Source/Editor/AddContentDialog/Private/AddContentDialogStyle.cpp
@@ -86,6 +86,13 @@ TSharedRef< FSlateStyleSet > FAddContentDialogStyle::Create()
 	Style->Set("AddContentDialog.ContentPackCategory", new IMAGE_BRUSH( "ContentPack", FVector2D(32.0f, 32.0f)));
 	Style->Set("AddContentDialog.UnknownCategory", new FSlateNoResource());
 
+	Style->Set("AddContentDialog.AddButton.TextStyle", FTextBlockStyle(NormalText)
+		.SetFont(TTF_CORE_FONT("Fonts/Roboto-Bold", 11))
+		.SetColorAndOpacity(FLinearColor(1.0f, 1.0f, 1.0f))
+		.SetHighlightColor(FLinearColor(1.0f, 1.0f, 1.0f))
+		.SetShadowOffset(FVector2D(1, 1))
+		.SetShadowColorAndOpacity(FLinearColor(0, 0, 0, 0.9f)));
+
 	return Style;
 }
 
diff --git a/Engine/Source/Editor/AddContentDialog/Private/SAddContentWidget.cpp b/Engine/Source/Editor/AddContentDialog/Private/SAddContentWidget.cpp
--- a/Engine/Source/Editor/AddContentDialog/Private/SAddContentWidget.cpp
+++ b/Engine/Source/Editor/AddContentDialog/Private/SAddContentWidget.cpp
@@ -244,13 +244,35 @@ TSharedRef<SWidget> SAddContentWidget::CreateContentSourceDetail(TSharedPtr<FCon
 	
 		VerticalBox->AddSlot()
 		.AutoHeight()
+		.Padding(0, 10, 0, 0)
+		.HAlign(HAlign_Right)
 		[
 			SNew(SButton)
+			.ButtonStyle(FEditorStyle::Get(), "FlatButton.Success")
 			.OnClicked(this, &SAddContentWidget::AddButtonClicked)
+			.ContentPadding(FMargin(5, 5, 5, 5))
 			.HAlign(EHorizontalAlignment::HAlign_Center)
 			[
-				SNew(STextBlock)
-				.Text(LOCTEXT("AddToProjectButton", "Add to Project"))
+				SNew(SHorizontalBox)
+
+				+ SHorizontalBox::Slot()
+				.VAlign(VAlign_Center)
+				.AutoWidth()
+				.Padding(0, 0, 2, 0)
+				[
+					SNew(STextBlock)
+					.TextStyle(FEditorStyle::Get(), "NormalText.Important")
+					.Font(FEditorStyle::Get().GetFontStyle("FontAwesome.10"))
+					.Text(FText::FromString(FString(TEXT("\xf067"))) /*fa-plus*/)
+				]
+
+				+ SHorizontalBox::Slot()
+				.AutoWidth()
+				[
+					SNew(STextBlock)
+					.TextStyle(FAddContentDialogStyle::Get(), "AddContentDialog.AddButton.TextStyle")
+					.Text(LOCTEXT("AddToProjectButton", "Add to Project"))
+				]
 			]
 		];
 	}
