AddContentDialog - Add error checking and validation for pack manifest loading and fix up some bad manifests.

diff --git a/Engine/Source/Editor/AddContentDialog/Private/ContentSourceProviders/FeaturePack/FeaturePackContentSource.cpp b/Engine/Source/Editor/AddContentDialog/Private/ContentSourceProviders/FeaturePack/FeaturePackContentSource.cpp
--- a/Engine/Source/Editor/AddContentDialog/Private/ContentSourceProviders/FeaturePack/FeaturePackContentSource.cpp
+++ b/Engine/Source/Editor/AddContentDialog/Private/ContentSourceProviders/FeaturePack/FeaturePackContentSource.cpp
@@ -10,6 +10,118 @@
 
 DEFINE_LOG_CATEGORY_STATIC(LogFeaturePack, Log, All);
 
+bool TryValidateTranslatedValue(TSharedPtr<FJsonValue> TranslatedValue, TSharedPtr<FString>& ErrorMessage)
+{
+	if (TranslatedValue.IsValid() == false)
+	{
+		ErrorMessage = MakeShareable(new FString("Invalid translated value"));
+		return false;
+	}
+
+	const TSharedPtr<FJsonObject>* TranslatedObject;
+	if (TranslatedValue->TryGetObject(TranslatedObject) == false)
+	{
+		ErrorMessage = MakeShareable(new FString("Invalid translated value"));
+		return false;
+	}
+
+	if ((*TranslatedObject)->HasTypedField<EJson::String>("Language") == false)
+	{
+		ErrorMessage = MakeShareable(new FString("Translated value missing 'Language' field"));
+		return false;
+	}
+
+	if ((*TranslatedObject)->HasTypedField<EJson::String>("Text") == false)
+	{
+		ErrorMessage = MakeShareable(new FString("Translated value missing 'Text' field"));
+		return false;
+	}
+
+	return true;
+}
+
+bool TryValidateManifestObject(TSharedPtr<FJsonObject> ManifestObject, TSharedPtr<FString>& ErrorMessage)
+{
+	if (ManifestObject.IsValid() == false)
+	{
+		ErrorMessage = MakeShareable(new FString("Manifest object missing"));
+		return false;
+	}
+
+	if (ManifestObject->HasTypedField<EJson::Array>("Name") == false)
+	{
+		ErrorMessage = MakeShareable(new FString("Manifest object missing 'Names' field"));
+		return false;
+	}
+
+	for (TSharedPtr<FJsonValue> NameValue : ManifestObject->GetArrayField("Name"))
+	{
+		if (TryValidateTranslatedValue(NameValue, ErrorMessage) == false)
+		{
+			ErrorMessage = MakeShareable(new FString(FString::Printf(TEXT("Manifest object 'Names' field error: %s"), **ErrorMessage)));
+		}
+	}
+
+	if (ManifestObject->HasTypedField<EJson::Array>("Description") == false)
+	{
+		ErrorMessage = MakeShareable(new FString("Manifest object missing 'Description' field"));
+		return false;
+	}
+
+	for (TSharedPtr<FJsonValue> DescriptionValue : ManifestObject->GetArrayField("Description"))
+	{
+		if (TryValidateTranslatedValue(DescriptionValue, ErrorMessage) == false)
+		{
+			ErrorMessage = MakeShareable(new FString(FString::Printf(TEXT("Manifest object 'Description' field error: %s"), **ErrorMessage)));
+		}
+	}
+
+	if (ManifestObject->HasTypedField<EJson::Array>("AssetTypes") == false)
+	{
+		ErrorMessage = MakeShareable(new FString("Manifest object missing 'AssetTypes' field"));
+		return false;
+	}
+
+	for (TSharedPtr<FJsonValue> AssetTypesValue : ManifestObject->GetArrayField("AssetTypes"))
+	{
+		if (TryValidateTranslatedValue(AssetTypesValue, ErrorMessage) == false)
+		{
+			ErrorMessage = MakeShareable(new FString(FString::Printf(TEXT("Manifest object 'AssetTypes' field error: %s"), **ErrorMessage)));
+		}
+	}
+
+	if (ManifestObject->HasTypedField<EJson::String>("ClassTypes") == false)
+	{
+		ErrorMessage = MakeShareable(new FString("Manifest object missing 'ClassTypes' field"));
+		return false;
+	}
+
+	if (ManifestObject->HasTypedField<EJson::String>("Category") == false)
+	{
+		ErrorMessage = MakeShareable(new FString("Manifest object missing 'Category' field"));
+		return false;
+	}
+
+	if (ManifestObject->HasTypedField<EJson::String>("FocusAsset") == false)
+	{
+		ErrorMessage = MakeShareable(new FString("Manifest object missing 'FocusAsset' field"));
+		return false;
+	}
+
+	if (ManifestObject->HasTypedField<EJson::String>("Thumbnail") == false)
+	{
+		ErrorMessage = MakeShareable(new FString("Manifest object missing 'Thumbnail' field"));
+		return false;
+	}
+
+	if (ManifestObject->HasTypedField<EJson::Array>("Screenshots") == false)
+	{
+		ErrorMessage = MakeShareable(new FString("Manifest object missing 'Screenshots' field"));
+		return false;
+	}
+
+	return true;
+}
 
 FFeaturePackContentSource::FFeaturePackContentSource(FString InFeaturePackPath)
 {
@@ -37,6 +149,22 @@ FFeaturePackContentSource::FFeaturePackContentSource(FString InFeaturePackPath)
 	TSharedPtr<FJsonObject> ManifestObject;
 	TSharedRef<TJsonReader<>> ManifestReader = TJsonReaderFactory<>::Create(ManifestString);
 	FJsonSerializer::Deserialize(ManifestReader, ManifestObject);
+
+	if (ManifestReader->GetErrorMessage().IsEmpty() == false)
+	{
+		UE_LOG(LogFeaturePack, Warning, TEXT("Error in Feature pack %s. Failed to parse manifest: %s"), *InFeaturePackPath, *ManifestReader->GetErrorMessage());
+		Category = EContentSourceCategory::Unknown;
+		return;
+	}
+
+	TSharedPtr<FString> ManifestObjectErrorMessage;
+	if (TryValidateManifestObject(ManifestObject, ManifestObjectErrorMessage) == false)
+	{
+		UE_LOG(LogFeaturePack, Warning, TEXT("Error in Feature pack %s. Manifest object error: %s"), *InFeaturePackPath, **ManifestObjectErrorMessage);
+		Category = EContentSourceCategory::Unknown;
+		return;
+	}
+
 	for (TSharedPtr<FJsonValue> NameValue : ManifestObject->GetArrayField("Name"))
 	{
 		TSharedPtr<FJsonObject> LocalizedNameObject = NameValue->AsObject();
