Added suport for classes to register with super search for handling/acting on searches
Hooked in feature packs to show in super search using tags in the FP manifests
Added some keywords to manifests for super search

diff --git a/Engine/Source/Developer/SuperSearch/Private/SSuperSearch.cpp b/Engine/Source/Developer/SuperSearch/Private/SSuperSearch.cpp
--- a/Engine/Source/Developer/SuperSearch/Private/SSuperSearch.cpp
+++ b/Engine/Source/Developer/SuperSearch/Private/SSuperSearch.cpp
@@ -161,20 +161,32 @@ void SSuperSearchBox::ActOnSuggestion(TSharedPtr<FSearchEntry> SearchEntry, FStr
 	{
 		EntryClicked = SearchEntry;
 #if WITH_EDITOR
+		// Broadcast text change to anyone registered for it
+		FSuperSearchModule& SuperSearchModule = FModuleManager::LoadModuleChecked< FSuperSearchModule >(TEXT("SuperSearch"));		
+		SuperSearchModule.GetActOnSearchTextClicked().Broadcast(SearchEntry);
+
+		// See if the search has tutorial hits
 		if (SearchEntry->URL.IsEmpty())
 		{
 			UBlueprint* Blueprint = LoadObject<UBlueprint>(nullptr, *(SearchEntry->AssetData.ObjectPath.ToString()));
-			UEditorTutorial* Tutorial = Blueprint->GeneratedClass->GetDefaultObject<UEditorTutorial>();
-			IIntroTutorials& IntroTutorials = IIntroTutorials::Get();
-			
-			FWidgetPath OutWidgetPath;
-			TSharedPtr<SWindow> ParentWindow = FSlateApplication::Get().FindWidgetWindow(SuggestionBox.ToSharedRef(), OutWidgetPath);
+			if( Blueprint != nullptr )
+			{
+				UEditorTutorial* Tutorial = Blueprint->GeneratedClass->GetDefaultObject<UEditorTutorial>();
+				if( Tutorial!= nullptr)
+				{
+					IIntroTutorials& IntroTutorials = IIntroTutorials::Get();
+
+					FWidgetPath OutWidgetPath;
+					TSharedPtr<SWindow> ParentWindow = FSlateApplication::Get().FindWidgetWindow(SuggestionBox.ToSharedRef(), OutWidgetPath);
 
-			IntroTutorials.LaunchTutorial(Tutorial, IIntroTutorials::ETutorialStartType::TST_RESTART, ParentWindow);
+					IntroTutorials.LaunchTutorial(Tutorial, IIntroTutorials::ETutorialStartType::TST_RESTART, ParentWindow);
+				}
+			} 
 		}
 		else
 #endif
 		{
+			// Check for documentation hits
 			if (Category == TEXT("documentation"))
 			{
 				// append some tracking data to the URL
@@ -537,14 +549,7 @@ FReply SSuperSearchBox::OnKeyDown( const FGeometry& MyGeometry, const FKeyEvent&
 	return FReply::Unhandled();
 }
 
-FSearchEntry * FSearchEntry::MakeCategoryEntry(const FString & InTitle)
-{
-	FSearchEntry * SearchEntry = new FSearchEntry();
-	SearchEntry->Title = InTitle;
-	SearchEntry->bCategory = true;
 
-	return SearchEntry;
-}
 
 void UpdateSuggestionHelper(const FText & CategoryLabel, const TArray<FSearchEntry> & Elements, TArray<TSharedPtr< FSearchEntry > > & OutSuggestions)
 {
@@ -589,6 +594,10 @@ void SSuperSearchBox::UpdateSuggestions()
 	//then answerhub
 	UpdateSuggestionHelper(NSLOCTEXT("SuperSearch", "answers", "Answerhub"), SearchResults->OnlineResults.FindOrAdd(TEXT("answers")), Suggestions);
 
+	FSuperSearchModule& SuperSearchModule = FModuleManager::LoadModuleChecked< FSuperSearchModule >(TEXT("SuperSearch"));
+	//Broadcast to anyone registered 
+	SuperSearchModule.GetSearchTextChanged().Broadcast(InputText->GetText().ToString(),Suggestions);
+
 	//finally add other category
 	Suggestions.Add(OtherCategory);
 	Suggestions.Add(AskQuestionEntry);
diff --git a/Engine/Source/Developer/SuperSearch/Private/SSuperSearch.h b/Engine/Source/Developer/SuperSearch/Private/SSuperSearch.h
--- a/Engine/Source/Developer/SuperSearch/Private/SSuperSearch.h
+++ b/Engine/Source/Developer/SuperSearch/Private/SSuperSearch.h
@@ -6,22 +6,10 @@
 #include "AssetData.h"
 #endif
 
-/**
- * Console input box with command-completion support
+/* 
+ * Super search - Searches a number of resources including documentation, tutorials, Wiki and AnswerHub
  */
 
-struct FSearchEntry
-{
-	FString Title;
-	FString URL;
-	bool bCategory;
-#if WITH_EDITOR
-	FAssetData AssetData;
-#endif
-	static FSearchEntry * MakeCategoryEntry(const FString & InTitle);
-};
-
-
 class SSuperSearchBox
 	: public SCompoundWidget
 {
@@ -56,7 +44,7 @@ public:
 	{
 		return InputText.ToSharedRef();
 	}
-
+		
 protected:
 
 	virtual bool SupportsKeyboardFocus() const { return true; }
@@ -90,7 +78,6 @@ protected:
 
 	FString GetSelectionText() const;
 
-
 private:
 
 	/** Editable text widget */
diff --git a/Engine/Source/Developer/SuperSearch/Private/SuperSearchModule.cpp b/Engine/Source/Developer/SuperSearch/Private/SuperSearchModule.cpp
--- a/Engine/Source/Developer/SuperSearch/Private/SuperSearchModule.cpp
+++ b/Engine/Source/Developer/SuperSearch/Private/SuperSearchModule.cpp
@@ -10,6 +10,16 @@ namespace SuperSearchModule
 	static const FName FNameSuperSearchApp = FName(TEXT("SuperSearchApp"));
 }
 
+FSearchEntry * FSearchEntry::MakeCategoryEntry(const FString & InTitle)
+{
+	FSearchEntry * SearchEntry = new FSearchEntry();
+	SearchEntry->Title = InTitle;
+	SearchEntry->bCategory = true;
+
+	return SearchEntry;
+}
+
+
 void FSuperSearchModule::StartupModule()
 {
 }
diff --git a/Engine/Source/Developer/SuperSearch/Public/SuperSearchModule.h b/Engine/Source/Developer/SuperSearch/Public/SuperSearchModule.h
--- a/Engine/Source/Developer/SuperSearch/Public/SuperSearchModule.h
+++ b/Engine/Source/Developer/SuperSearch/Public/SuperSearchModule.h
@@ -4,13 +4,52 @@
 
 #include "SlateBasics.h"
 
+struct FSearchEntry
+{
+	FString Title;
+	FString URL;
+	bool bCategory;
+#if WITH_EDITOR
+	FAssetData AssetData;
+#endif
+	static FSearchEntry* MakeCategoryEntry(const FString & InTitle);
+};
+
+
+
 class FSuperSearchModule : public IModuleInterface
 {
 public:
+	/* 
+	 * Delegate for search text changing.
+	 * @param Search string
+	 * @param Suggestion array
+	 */
+	DECLARE_MULTICAST_DELEGATE_TwoParams(FSuperSearchTextChanged, const FString&, TArray< TSharedPtr<FSearchEntry> >&);
+	/* 
+	 * Delegate for search text click.
+	 * @param Search entry clicked.
+	 */	
+	DECLARE_MULTICAST_DELEGATE_OneParam(FActOnSuperSearchEntry, TSharedPtr<FSearchEntry> );
+
 	virtual void StartupModule();
 	virtual void ShutdownModule();
 
-	/** Generates aSuperSearch box widget.  Remember, this widget will become invalid if the
-	   SuperSearch DLL is unloaded on the fly. */
+	// Get the search text changed delegate
+	FSuperSearchTextChanged& GetSearchTextChanged()
+	{
+		return OnSearchTextChanged;
+	}
+	// Get the search hit clicked changed delegate
+	FActOnSuperSearchEntry& GetActOnSearchTextClicked()
+	{
+		return OnActOnSearchEntry;
+	}
+	
+	/** Generates a SuperSearch box widget.  Remember, this widget will become invalid if the SuperSearch DLL is unloaded on the fly. */
 	virtual TSharedRef< SWidget > MakeSearchBox(TSharedPtr< SEditableTextBox >& OutExposedEditableTextBox, const TOptional<const FSearchBoxStyle*> InStyle = TOptional<const FSearchBoxStyle*>()) const;
+private:
+	
+	FSuperSearchTextChanged	OnSearchTextChanged;	
+	FActOnSuperSearchEntry OnActOnSearchEntry;
 };
diff --git a/Engine/Source/Editor/AddContentDialog/AddContentDialog.Build.cs b/Engine/Source/Editor/AddContentDialog/AddContentDialog.Build.cs
--- a/Engine/Source/Editor/AddContentDialog/AddContentDialog.Build.cs
+++ b/Engine/Source/Editor/AddContentDialog/AddContentDialog.Build.cs
@@ -10,6 +10,7 @@ public class AddContentDialog : ModuleRules
 			new string[] {
 				"AssetTools",
 				"ContentBrowser",
+				"SuperSearch"
 			}
 		);
 
@@ -26,18 +27,15 @@ public class AddContentDialog : ModuleRules
 				"PakFile",
 				"ImageWrapper",
 				"UnrealEd",
-				"CoreUObject"
-			}
-		);
-
-		PrivateDependencyModuleNames.AddRange(
-			new string[] {
-				"WidgetCarousel"
+				"CoreUObject",				
+				"WidgetCarousel",				
+			
 			}
 		);
 
 		DynamicallyLoadedModuleNames.AddRange(
 			new string[] {
+				"SuperSearch",
 			}
 		);
 		
diff --git a/Engine/Source/Editor/AddContentDialog/Private/ContentSourceProviders/FeaturePack/FeaturePackContentSource.cpp b/Engine/Source/Editor/AddContentDialog/Private/ContentSourceProviders/FeaturePack/FeaturePackContentSource.cpp
--- a/Engine/Source/Editor/AddContentDialog/Private/ContentSourceProviders/FeaturePack/FeaturePackContentSource.cpp
+++ b/Engine/Source/Editor/AddContentDialog/Private/ContentSourceProviders/FeaturePack/FeaturePackContentSource.cpp
@@ -7,6 +7,10 @@
 #include "ContentBrowserModule.h"
 #include "IPlatformFilePak.h"
 #include "FileHelpers.h"
+#include "Editor/MainFrame/Public/Interfaces/IMainFrameModule.h"
+/*#include "SuperSearchModule.h"*/
+
+#define LOCTEXT_NAMESPACE "ContentFeaturePacks"
 
 DEFINE_LOG_CATEGORY_STATIC(LogFeaturePack, Log, All);
 
@@ -190,6 +194,15 @@ FFeaturePackContentSource::FFeaturePackContentSource(FString InFeaturePackPath)
 			FText::FromString(LocalizedAssetTypesObject->GetStringField("Text"))));
 	}
 	
+	// Parse asset types field
+	for (TSharedPtr<FJsonValue> AssetTypesValue : ManifestObject->GetArrayField("SearchTags"))
+	{
+		TSharedPtr<FJsonObject> LocalizedAssetTypesObject = AssetTypesValue->AsObject();
+		LocalizedSearchTags.Add(FLocalizedTextArray(
+			LocalizedAssetTypesObject->GetStringField("Language"),
+			LocalizedAssetTypesObject->GetStringField("Text")));
+	}
+
 	// Parse class types field
 	ClassTypes = ManifestObject->GetStringField("ClassTypes");
 	
@@ -201,22 +214,9 @@ FFeaturePackContentSource::FFeaturePackContentSource(FString InFeaturePackPath)
 	ManifestObject->TryGetStringField("SortKey", SortKey);
 
 	FString CategoryString = ManifestObject->GetStringField("Category");	
-	if (CategoryString == "CodeFeature")
-	{
-		Category = EContentSourceCategory::CodeFeature;
-	}
-	else if (CategoryString == "BlueprintFeature")
-	{
-		Category = EContentSourceCategory::BlueprintFeature;
-	}
-	else if (CategoryString == "Content")
-	{
-		Category = EContentSourceCategory::Content;
-	}
-	else
-	{
-		Category = EContentSourceCategory::Unknown;
-	}
+	UEnum* Enum = FindObjectChecked<UEnum>(ANY_PACKAGE, TEXT("EContentSourceCategory"));
+	int32 Index = Enum->FindEnumIndex(FName(*CategoryString));
+	Category = Index != INDEX_NONE ? (EContentSourceCategory)Index : EContentSourceCategory::Unknown;
 
 	// Load image data
 	FString IconFilename = ManifestObject->GetStringField("Thumbnail");
@@ -238,6 +238,10 @@ FFeaturePackContentSource::FFeaturePackContentSource(FString InFeaturePackPath)
 		LoadPakFileToBuffer(PakPlatformFile, FPaths::Combine(*MountPoint, TEXT("Media"), *ScreenshotFilename->AsString()), *SingleScreenshotData);
 		ScreenshotData.Add(MakeShareable(new FImageData(ScreenshotFilename->AsString(), SingleScreenshotData)));
 	}
+
+	FSuperSearchModule& SuperSearchModule = FModuleManager::LoadModuleChecked< FSuperSearchModule >(TEXT("SuperSearch"));
+	SuperSearchModule.GetActOnSearchTextClicked().AddRaw(this, &FFeaturePackContentSource::HandleActOnSearchText);	
+	SuperSearchModule.GetSearchTextChanged().AddRaw(this, &FFeaturePackContentSource::HandleSuperSearchTextChanged);
 	bPackValid = true;
 }
 
@@ -253,37 +257,38 @@ bool FFeaturePackContentSource::LoadPakFileToBuffer(FPakPlatformFile& PakPlatfor
 	return bResult;
 }
 
-TArray<FLocalizedText> FFeaturePackContentSource::GetLocalizedNames()
+TArray<FLocalizedText> FFeaturePackContentSource::GetLocalizedNames() const
 {
 	return LocalizedNames;
 }
 
-TArray<FLocalizedText> FFeaturePackContentSource::GetLocalizedDescriptions()
+TArray<FLocalizedText> FFeaturePackContentSource::GetLocalizedDescriptions() const
 {
 	return LocalizedDescriptions;
 }
 
-TArray<FLocalizedText> FFeaturePackContentSource::GetLocalizedAssetTypes()
+TArray<FLocalizedText> FFeaturePackContentSource::GetLocalizedAssetTypes() const
 {
 	return LocalizedAssetTypesList;
 }
 
-FString FFeaturePackContentSource::GetClassTypesUsed()
+
+FString FFeaturePackContentSource::GetClassTypesUsed() const
 {
 	return ClassTypes;
 }
 
-EContentSourceCategory FFeaturePackContentSource::GetCategory()
+EContentSourceCategory FFeaturePackContentSource::GetCategory() const
 {
 	return Category;
 }
 
-TSharedPtr<FImageData> FFeaturePackContentSource::GetIconData()
+TSharedPtr<FImageData> FFeaturePackContentSource::GetIconData() const
 {
 	return IconData;
 }
 
-TArray<TSharedPtr<FImageData>> FFeaturePackContentSource::GetScreenshotData()
+TArray<TSharedPtr<FImageData>> FFeaturePackContentSource::GetScreenshotData() const
 {
 	return ScreenshotData;
 }
@@ -356,4 +361,109 @@ FString FFeaturePackContentSource::GetFocusAssetName() const
 FString FFeaturePackContentSource::GetSortKey() const
 {
 	return SortKey;
-}
\ No newline at end of file
+}
+
+
+void FFeaturePackContentSource::HandleActOnSearchText(TSharedPtr<FSearchEntry> SearchEntry)
+{
+	if (SearchEntry.IsValid())
+	{
+		if (SearchEntry->bCategory == false)
+		{
+			UEnum* Enum = FindObjectChecked<UEnum>(ANY_PACKAGE, TEXT("EContentSourceCategory"));
+			FString CurrentLanguage = FInternationalization::Get().GetCurrentCulture()->GetTwoLetterISOLanguageName();
+			FLocalizedText CurrentName = ChooseLocalizedText(LocalizedNames,CurrentLanguage);
+			FString MyTitle = FText::Format( LOCTEXT("FeaturePackSearchResult", "{0} ({1})"), CurrentName.GetText(), Enum->GetDisplayNameText((int32)Category)).ToString();
+			if (SearchEntry->Title == MyTitle)
+			{
+				IMainFrameModule& MainFrameModule = FModuleManager::LoadModuleChecked<IMainFrameModule>(TEXT("MainFrame"));
+				IAddContentDialogModule& AddContentDialogModule = FModuleManager::LoadModuleChecked<IAddContentDialogModule>("AddContentDialog");				
+				AddContentDialogModule.ShowDialog(MainFrameModule.GetParentWindow().ToSharedRef());
+			}
+		}
+	}
+}
+
+void FFeaturePackContentSource::TryAddFeaturePackCategory(FString CategoryTitle, TArray< TSharedPtr<FSearchEntry> >& OutSuggestions)
+{
+	bool bCategoryExists = false;
+	if( OutSuggestions.FindByPredicate(
+		[&CategoryTitle](TSharedPtr<FSearchEntry> InElement)
+	{
+		return (( InElement->Title == CategoryTitle ) && (InElement->bCategory == true ));
+	})==false)
+	{
+		TSharedPtr<FSearchEntry> FeaturePackCat = MakeShareable(new FSearchEntry());
+		FeaturePackCat->bCategory = true;
+		FeaturePackCat->Title = CategoryTitle;
+		OutSuggestions.Add(FeaturePackCat);
+
+	}
+}
+
+void FFeaturePackContentSource::HandleSuperSearchTextChanged(const FString& InText, TArray< TSharedPtr<FSearchEntry> >& OutSuggestions)
+{
+	FString FeaturePackSearchCat = LOCTEXT("FeaturePackSearchCategory", "Feature Packs").ToString();
+
+	FString CurrentLanguage = FInternationalization::Get().GetCurrentCulture()->GetTwoLetterISOLanguageName();
+	FLocalizedText CurrentName = ChooseLocalizedText(LocalizedNames,CurrentLanguage);
+	FLocalizedTextArray CurrentTagSet = ChooseLocalizedTextArray(LocalizedSearchTags,CurrentLanguage);
+	FText AsText = FText::FromString(InText);
+	TArray<FText> TagArray = CurrentTagSet.GetTags();
+	if (TagArray.Num() != 0)
+	{
+		UEnum* Enum = FindObjectChecked<UEnum>(ANY_PACKAGE, TEXT("EContentSourceCategory"));
+		// Add a feature packs category		
+		for (int32 iTag = 0; iTag < TagArray.Num(); iTag++)
+		{
+			if (TagArray[iTag].EqualToCaseIgnored(AsText))
+			{
+				// This will add the category if one doesnt exist
+				TryAddFeaturePackCategory(FeaturePackSearchCat,OutSuggestions);
+				TSharedPtr<FSearchEntry> FeaturePackEntry = MakeShareable(new FSearchEntry());
+				FeaturePackEntry->Title = FText::Format( LOCTEXT("FeaturePackSearchResult", "{0} ({1})"), CurrentName.GetText(), Enum->GetDisplayNameText((int32)Category)).ToString();
+				FeaturePackEntry->bCategory = false;
+				OutSuggestions.Add(FeaturePackEntry);
+				return;
+			}
+		}
+	}
+}
+
+FLocalizedText FFeaturePackContentSource::ChooseLocalizedText(TArray<FLocalizedText> Choices, FString LanguageCode)
+{
+	FLocalizedText Default;
+	for (const FLocalizedText& Choice : Choices)
+	{
+		if (Choice.GetTwoLetterLanguage() == LanguageCode)
+		{
+			return Choice;
+			break;
+		}
+		else if (Choice.GetTwoLetterLanguage() == TEXT("en"))
+		{
+			Default = Choice;
+		}
+	}
+	return Default;
+}
+
+FLocalizedTextArray FFeaturePackContentSource::ChooseLocalizedTextArray(TArray<FLocalizedTextArray> Choices, FString LanguageCode)
+{
+	FLocalizedTextArray Default;
+	for (const FLocalizedTextArray& Choice : Choices)
+	{
+		if (Choice.GetTwoLetterLanguage() == LanguageCode)
+		{
+			return Choice;
+			break;
+		}
+		else if (Choice.GetTwoLetterLanguage() == TEXT("en"))
+		{
+			Default = Choice;
+		}
+	}
+	return Default;
+}
+
+#undef LOCTEXT_NAMESPACE 
diff --git a/Engine/Source/Editor/AddContentDialog/Private/ContentSourceProviders/FeaturePack/FeaturePackContentSource.h b/Engine/Source/Editor/AddContentDialog/Private/ContentSourceProviders/FeaturePack/FeaturePackContentSource.h
--- a/Engine/Source/Editor/AddContentDialog/Private/ContentSourceProviders/FeaturePack/FeaturePackContentSource.h
+++ b/Engine/Source/Editor/AddContentDialog/Private/ContentSourceProviders/FeaturePack/FeaturePackContentSource.h
@@ -1,6 +1,46 @@
 // Copyright 1998-2015 Epic Games, Inc. All Rights Reserved.
 
 #pragma once
+#include "SuperSearchModule.h"
+
+class FLocalizedTextArray
+{
+public:
+	FLocalizedTextArray()
+	{
+	}
+
+	/** Creates a new FLocalizedText
+		@param InTwoLetterLanguage - The iso 2-letter language specifier.
+		@param InText - The text in the language specified */
+	FLocalizedTextArray(FString InTwoLetterLanguage, FString InText)
+	{
+		TwoLetterLanguage = InTwoLetterLanguage;
+		TArray<FString> AsArray;
+		InText.ParseIntoArray(AsArray,TEXT(","));
+		for (int32 iString = 0; iString < AsArray.Num() ; iString++)
+		{
+			Tags.Add(FText::FromString(AsArray[iString]));
+		}
+	}
+
+	/** Gets the iso 2-letter language specifier for this text. */
+	FString GetTwoLetterLanguage() const
+	{
+		return TwoLetterLanguage;
+	}
+
+	/** Gets the array of tags in the language specified. */
+	TArray<FText> GetTags() const
+	{
+		return Tags;
+	}
+
+private:
+	FString TwoLetterLanguage;
+	TArray<FText> Tags;
+};
+
 
 class FPakPlatformFile;
 
@@ -10,27 +50,35 @@ class FFeaturePackContentSource : public IContentSource
 public:
 	FFeaturePackContentSource(FString InFeaturePackPath);
 
-	virtual TArray<FLocalizedText> GetLocalizedNames() override;
-	virtual TArray<FLocalizedText> GetLocalizedDescriptions() override;
+	virtual TArray<FLocalizedText> GetLocalizedNames() const override;
+	virtual TArray<FLocalizedText> GetLocalizedDescriptions() const override;
 	
-	virtual EContentSourceCategory GetCategory() override;
-	virtual TArray<FLocalizedText> GetLocalizedAssetTypes() override;
+	virtual EContentSourceCategory GetCategory() const override;
+	virtual TArray<FLocalizedText> GetLocalizedAssetTypes() const override;
 	virtual FString GetSortKey() const override;
-	virtual FString GetClassTypesUsed() override;
-
-	virtual TSharedPtr<FImageData> GetIconData() override;
-	virtual TArray<TSharedPtr<FImageData>> GetScreenshotData() override;
+	virtual FString GetClassTypesUsed() const override;
+	
+	virtual TSharedPtr<FImageData> GetIconData() const override;
+	virtual TArray<TSharedPtr<FImageData>> GetScreenshotData() const override;
 
 	virtual bool InstallToProject(FString InstallPath) override;
 	virtual bool IsDataValid() const override;
 	
 	virtual ~FFeaturePackContentSource();
 	
+	virtual void HandleActOnSearchText(TSharedPtr<FSearchEntry> SearchEntry);
+	virtual void HandleSuperSearchTextChanged(const FString& InText, TArray< TSharedPtr<FSearchEntry> >& OutSuggestions);
+
 private:
 	bool LoadPakFileToBuffer(FPakPlatformFile& PakPlatformFile, FString Path, TArray<uint8>& Buffer);
 	FString GetFocusAssetName() const;
 
 private:
+	/** Selects an FLocalizedText from an array which matches either the supplied language code, or the default language code. */
+	FLocalizedTextArray ChooseLocalizedTextArray(TArray<FLocalizedTextArray> Choices, FString LanguageCode);
+	FLocalizedText ChooseLocalizedText(TArray<FLocalizedText> Choices, FString LanguageCode);
+	void TryAddFeaturePackCategory(FString CategoryTitle, TArray< TSharedPtr<FSearchEntry> >& OutSuggestions);
+
 	FString FeaturePackPath;
 	TArray<FLocalizedText> LocalizedNames;
 	TArray<FLocalizedText> LocalizedDescriptions;
@@ -42,4 +90,5 @@ private:
 	bool bPackValid;
 	FString FocusAssetIdent;
 	FString SortKey;
+	TArray<FLocalizedTextArray> LocalizedSearchTags;
 };
\ No newline at end of file
diff --git a/Engine/Source/Editor/AddContentDialog/Private/IContentSource.h b/Engine/Source/Editor/AddContentDialog/Private/IContentSource.h
--- a/Engine/Source/Editor/AddContentDialog/Private/IContentSource.h
+++ b/Engine/Source/Editor/AddContentDialog/Private/IContentSource.h
@@ -2,8 +2,11 @@
 
 #pragma once
 
+#include "IContentSource.generated.h"
+
 /** Defines categories for content sources. */
-enum class EContentSourceCategory
+UENUM()
+enum class EContentSourceCategory:uint8
 {
 	BlueprintFeature,
 	CodeFeature,
@@ -11,6 +14,7 @@ enum class EContentSourceCategory
 	Unknown
 };
 
+
 /** Represents raw binary image data in png format. */
 class FImageData
 {
@@ -80,31 +84,31 @@ class IContentSource
 {
 public:
 	/** Gets the name of the content source as an array of localized strings. */
-	virtual TArray<FLocalizedText> GetLocalizedNames() = 0;
+	virtual TArray<FLocalizedText> GetLocalizedNames() const = 0;
 
 	/** Gets the description of the content source as an array or localized strings. */
-	virtual TArray<FLocalizedText> GetLocalizedDescriptions() = 0;
+	virtual TArray<FLocalizedText> GetLocalizedDescriptions() const = 0;
 
 	/** Gets the category for the content source. */
-	virtual EContentSourceCategory GetCategory() = 0;
+	virtual EContentSourceCategory GetCategory() const = 0;
 
 	/** Gets the image data for the icon which should represent the content source in the UI. */
-	virtual TSharedPtr<FImageData> GetIconData() = 0;
+	virtual TSharedPtr<FImageData> GetIconData() const = 0;
 
 	/** Gets an array or image data for screenshots for the content source. */
-	virtual TArray<TSharedPtr<FImageData>> GetScreenshotData() = 0;
+	virtual TArray<TSharedPtr<FImageData>> GetScreenshotData() const = 0;
 
 	/** Gets the asset types used in this pack. */
-	virtual TArray<FLocalizedText> GetLocalizedAssetTypes() = 0;
+	virtual TArray<FLocalizedText> GetLocalizedAssetTypes() const = 0;
 
 	/** Gets the class types used in this pack. */
-	virtual FString GetClassTypesUsed() = 0;
+	virtual FString GetClassTypesUsed() const = 0;
 
 	/** Gets the category for the content source. */
-	virtual FString GetSortKey() const { return FString(TEXT(""));};
+	virtual FString GetSortKey() const = 0;
 	/*
 	 * Installs the content in the content source to the specific path. 
-	 * @returns true if install suceeded
+	 * @returns true if install succeeded
 	 */
 	virtual bool InstallToProject(FString InstallPath) = 0;
 
diff --git a/Engine/Source/Editor/ContentBrowser/Private/SContentBrowser.cpp b/Engine/Source/Editor/ContentBrowser/Private/SContentBrowser.cpp
--- a/Engine/Source/Editor/ContentBrowser/Private/SContentBrowser.cpp
+++ b/Engine/Source/Editor/ContentBrowser/Private/SContentBrowser.cpp
@@ -626,6 +626,9 @@ void SContentBrowser::Construct( const FArguments& InArgs, const FName& InInstan
 	FAssetRegistryModule& AssetRegistryModule = FModuleManager::LoadModuleChecked<FAssetRegistryModule>(TEXT("AssetRegistry"));
 	AssetRegistryModule.Get().OnPathRemoved().AddSP(this, &SContentBrowser::HandlePathRemoved);
 
+	// We want to be able to search the feature packs in the super search so we need the module loaded 
+	IAddContentDialogModule& AddContentDialogModule = FModuleManager::LoadModuleChecked<IAddContentDialogModule>("AddContentDialog");
+
 	const TWeakPtr<SContentBrowser> WeakThis = SharedThis(this);
 
 	FContentBrowserModule& ContentBrowserModule = FModuleManager::GetModuleChecked<FContentBrowserModule>( TEXT("ContentBrowser") );
