Mac compile fix for FeaturePackContentSource

diff --git a/Engine/Source/Editor/AddContentDialog/Private/ContentSourceProviders/FeaturePack/FeaturePackContentSource.cpp b/Engine/Source/Editor/AddContentDialog/Private/ContentSourceProviders/FeaturePack/FeaturePackContentSource.cpp
--- a/Engine/Source/Editor/AddContentDialog/Private/ContentSourceProviders/FeaturePack/FeaturePackContentSource.cpp
+++ b/Engine/Source/Editor/AddContentDialog/Private/ContentSourceProviders/FeaturePack/FeaturePackContentSource.cpp
@@ -386,19 +386,18 @@ void FFeaturePackContentSource::HandleActOnSearchText(TSharedPtr<FSearchEntry> S
 
 void FFeaturePackContentSource::TryAddFeaturePackCategory(FString CategoryTitle, TArray< TSharedPtr<FSearchEntry> >& OutSuggestions)
 {
-	bool bCategoryExists = false;
-	if( OutSuggestions.FindByPredicate(
-		[&CategoryTitle](TSharedPtr<FSearchEntry> InElement)
+	for (int32 iEntry = 0; iEntry < OutSuggestions.Num() ; iEntry++)
 	{
-		return (( InElement->Title == CategoryTitle ) && (InElement->bCategory == true ));
-	})==false)
-	{
-		TSharedPtr<FSearchEntry> FeaturePackCat = MakeShareable(new FSearchEntry());
-		FeaturePackCat->bCategory = true;
-		FeaturePackCat->Title = CategoryTitle;
-		OutSuggestions.Add(FeaturePackCat);
-
+		if((OutSuggestions[iEntry]->Title == CategoryTitle ) && (OutSuggestions[iEntry]->bCategory == true ))
+		{
+			return;
+		}
 	}
+	
+	TSharedPtr<FSearchEntry> FeaturePackCat = MakeShareable(new FSearchEntry());
+	FeaturePackCat->bCategory = true;
+	FeaturePackCat->Title = CategoryTitle;
+	OutSuggestions.Add(FeaturePackCat);
 }
 
 void FFeaturePackContentSource::HandleSuperSearchTextChanged(const FString& InText, TArray< TSharedPtr<FSearchEntry> >& OutSuggestions)
