Added missing file to contents for vehicle FP and repacked said FP.
Added log output to FeaturePackContentSource to show when a screenshot file is missing.

diff --git a/Engine/Source/Editor/AddContentDialog/Private/ContentSourceProviders/FeaturePack/FeaturePackContentSource.cpp b/Engine/Source/Editor/AddContentDialog/Private/ContentSourceProviders/FeaturePack/FeaturePackContentSource.cpp
--- a/Engine/Source/Editor/AddContentDialog/Private/ContentSourceProviders/FeaturePack/FeaturePackContentSource.cpp
+++ b/Engine/Source/Editor/AddContentDialog/Private/ContentSourceProviders/FeaturePack/FeaturePackContentSource.cpp
@@ -235,8 +235,14 @@ FFeaturePackContentSource::FFeaturePackContentSource(FString InFeaturePackPath)
 	for (const TSharedPtr<FJsonValue> ScreenshotFilename : ScreenshotFilenameArray)
 	{
 		TSharedPtr<TArray<uint8>> SingleScreenshotData = MakeShareable(new TArray<uint8>);
-		LoadPakFileToBuffer(PakPlatformFile, FPaths::Combine(*MountPoint, TEXT("Media"), *ScreenshotFilename->AsString()), *SingleScreenshotData);
-		ScreenshotData.Add(MakeShareable(new FImageData(ScreenshotFilename->AsString(), SingleScreenshotData)));
+		if( LoadPakFileToBuffer(PakPlatformFile, FPaths::Combine(*MountPoint, TEXT("Media"), *ScreenshotFilename->AsString()), *SingleScreenshotData) )
+		{
+			ScreenshotData.Add(MakeShareable(new FImageData(ScreenshotFilename->AsString(), SingleScreenshotData)));
+		}
+		else
+		{
+			UE_LOG(LogFeaturePack, Warning, TEXT("Error in Feature pack %s. Cannot find screenshot %s."), *InFeaturePackPath, *ScreenshotFilename->AsString() );
+		}
 	}
 
 	FSuperSearchModule& SuperSearchModule = FModuleManager::LoadModuleChecked< FSuperSearchModule >(TEXT("SuperSearch"));
