Removed code that was forcing feature packs have a focus asset. (Focus asses aret not mandatory)

diff --git a/Engine/Source/Editor/AddContentDialog/Private/ContentSourceProviders/FeaturePack/FeaturePackContentSource.cpp b/Engine/Source/Editor/AddContentDialog/Private/ContentSourceProviders/FeaturePack/FeaturePackContentSource.cpp
--- a/Engine/Source/Editor/AddContentDialog/Private/ContentSourceProviders/FeaturePack/FeaturePackContentSource.cpp
+++ b/Engine/Source/Editor/AddContentDialog/Private/ContentSourceProviders/FeaturePack/FeaturePackContentSource.cpp
@@ -105,13 +105,7 @@ bool TryValidateManifestObject(TSharedPtr<FJsonObject> ManifestObject, TSharedPt
 		ErrorMessage = MakeShareable(new FString("Manifest object missing 'Category' field"));
 		return false;
 	}
-
-	if (ManifestObject->HasTypedField<EJson::String>("FocusAsset") == false)
-	{
-		ErrorMessage = MakeShareable(new FString("Manifest object missing 'FocusAsset' field"));
-		return false;
-	}
-
+		
 	if (ManifestObject->HasTypedField<EJson::String>("Thumbnail") == false)
 	{
 		ErrorMessage = MakeShareable(new FString("Manifest object missing 'Thumbnail' field"));
@@ -206,8 +200,12 @@ FFeaturePackContentSource::FFeaturePackContentSource(FString InFeaturePackPath)
 	// Parse class types field
 	ClassTypes = ManifestObject->GetStringField("ClassTypes");
 	
-	// Parse initial focus asset
-	FocusAssetIdent = ManifestObject->GetStringField("FocusAsset");
+	
+	// Parse initial focus asset if we have one - this is not required
+	if (ManifestObject->HasTypedField<EJson::String>("FocusAsset") == false)
+	{
+		FocusAssetIdent = ManifestObject->GetStringField("FocusAsset");
+	}	
 
 	// Use the path as the sort key - it will be alphabetical that way
 	SortKey = FeaturePackPath;
@@ -392,18 +390,14 @@ void FFeaturePackContentSource::HandleActOnSearchText(TSharedPtr<FSearchEntry> S
 
 void FFeaturePackContentSource::TryAddFeaturePackCategory(FString CategoryTitle, TArray< TSharedPtr<FSearchEntry> >& OutSuggestions)
 {
-	for (int32 iEntry = 0; iEntry < OutSuggestions.Num() ; iEntry++)
+	if (OutSuggestions.ContainsByPredicate([&CategoryTitle](TSharedPtr<FSearchEntry>& InElement)
+		{ return ((InElement->Title == CategoryTitle) && (InElement->bCategory == true)); }) == false)
 	{
-		if((OutSuggestions[iEntry]->Title == CategoryTitle ) && (OutSuggestions[iEntry]->bCategory == true ))
-		{
-			return;
-		}
+		TSharedPtr<FSearchEntry> FeaturePackCat = MakeShareable(new FSearchEntry());
+		FeaturePackCat->bCategory = true;
+		FeaturePackCat->Title = CategoryTitle;
+		OutSuggestions.Add(FeaturePackCat);
 	}
-	
-	TSharedPtr<FSearchEntry> FeaturePackCat = MakeShareable(new FSearchEntry());
-	FeaturePackCat->bCategory = true;
-	FeaturePackCat->Title = CategoryTitle;
-	OutSuggestions.Add(FeaturePackCat);
 }
 
 void FFeaturePackContentSource::HandleSuperSearchTextChanged(const FString& InText, TArray< TSharedPtr<FSearchEntry> >& OutSuggestions)
