Fixed a few compile errors when compiling with Clang on Windows (r231657)

diff --git a/Engine/Source/Editor/AddContentDialog/Private/ContentSourceProviders/FeaturePack/FeaturePackContentSource.cpp b/Engine/Source/Editor/AddContentDialog/Private/ContentSourceProviders/FeaturePack/FeaturePackContentSource.cpp
--- a/Engine/Source/Editor/AddContentDialog/Private/ContentSourceProviders/FeaturePack/FeaturePackContentSource.cpp
+++ b/Engine/Source/Editor/AddContentDialog/Private/ContentSourceProviders/FeaturePack/FeaturePackContentSource.cpp
@@ -8,7 +8,7 @@
 #include "IPlatformFilePak.h"
 #include "FileHelpers.h"
 #include "Editor/MainFrame/Public/Interfaces/IMainFrameModule.h"
-/*#include "SuperSearchModule.h"*/
+#include "SuperSearchModule.h"
 
 #define LOCTEXT_NAMESPACE "ContentFeaturePacks"
 
diff --git a/Engine/Source/Editor/AddContentDialog/Private/ContentSourceProviders/FeaturePack/FeaturePackContentSource.h b/Engine/Source/Editor/AddContentDialog/Private/ContentSourceProviders/FeaturePack/FeaturePackContentSource.h
--- a/Engine/Source/Editor/AddContentDialog/Private/ContentSourceProviders/FeaturePack/FeaturePackContentSource.h
+++ b/Engine/Source/Editor/AddContentDialog/Private/ContentSourceProviders/FeaturePack/FeaturePackContentSource.h
@@ -1,7 +1,6 @@
 // Copyright 1998-2015 Epic Games, Inc. All Rights Reserved.
 
 #pragma once
-#include "SuperSearchModule.h"
 
 class FLocalizedTextArray
 {
@@ -43,6 +42,7 @@ private:
 
 
 class FPakPlatformFile;
+struct FSearchEntry;
 
 /** A content source which represents a content upack. */
 class FFeaturePackContentSource : public IContentSource
@@ -66,8 +66,8 @@ public:
 	
 	virtual ~FFeaturePackContentSource();
 	
-	virtual void HandleActOnSearchText(TSharedPtr<FSearchEntry> SearchEntry);
-	virtual void HandleSuperSearchTextChanged(const FString& InText, TArray< TSharedPtr<FSearchEntry> >& OutSuggestions);
+	void HandleActOnSearchText(TSharedPtr<FSearchEntry> SearchEntry);
+	void HandleSuperSearchTextChanged(const FString& InText, TArray< TSharedPtr<FSearchEntry> >& OutSuggestions);
 
 private:
 	bool LoadPakFileToBuffer(FPakPlatformFile& PakPlatformFile, FString Path, TArray<uint8>& Buffer);
diff --git a/Engine/Source/Runtime/Windows/UnrealAudioXAudio2/Private/UnrealAudioDeviceXAudio2.cpp b/Engine/Source/Runtime/Windows/UnrealAudioXAudio2/Private/UnrealAudioDeviceXAudio2.cpp
--- a/Engine/Source/Runtime/Windows/UnrealAudioXAudio2/Private/UnrealAudioDeviceXAudio2.cpp
+++ b/Engine/Source/Runtime/Windows/UnrealAudioXAudio2/Private/UnrealAudioDeviceXAudio2.cpp
@@ -27,10 +27,10 @@ static const TCHAR* GetXAudio2Error(HRESULT Result)
 {
 	switch (Result)
 	{
-		case XAUDIO2_E_INVALID_CALL:					return TEXT("XAUDIO2_E_INVALID_CALL");
-		case XAUDIO2_E_XMA_DECODER_ERROR:				return TEXT("XAUDIO2_E_XMA_DECODER_ERROR");
-		case XAUDIO2_E_XAPO_CREATION_FAILED:			return TEXT("XAUDIO2_E_XAPO_CREATION_FAILED");
-		case XAUDIO2_E_DEVICE_INVALIDATED:				return TEXT("XAUDIO2_E_DEVICE_INVALIDATED");
+		case HRESULT(XAUDIO2_E_INVALID_CALL):			return TEXT("XAUDIO2_E_INVALID_CALL");
+		case HRESULT(XAUDIO2_E_XMA_DECODER_ERROR):		return TEXT("XAUDIO2_E_XMA_DECODER_ERROR");
+		case HRESULT(XAUDIO2_E_XAPO_CREATION_FAILED):	return TEXT("XAUDIO2_E_XAPO_CREATION_FAILED");
+		case HRESULT(XAUDIO2_E_DEVICE_INVALIDATED):		return TEXT("XAUDIO2_E_DEVICE_INVALIDATED");
 		case REGDB_E_CLASSNOTREG:						return TEXT("REGDB_E_CLASSNOTREG");
 		case CLASS_E_NOAGGREGATION:						return TEXT("CLASS_E_NOAGGREGATION");
 		case E_NOINTERFACE:								return TEXT("E_NOINTERFACE");
