Cleaned up FP_FirstPersonBP DefaultEngine.ini
Wrapped search tags parse in FeaturePack code with a check that the field exists

diff --git a/Engine/Source/Editor/AddContentDialog/Private/ContentSourceProviders/FeaturePack/FeaturePackContentSource.cpp b/Engine/Source/Editor/AddContentDialog/Private/ContentSourceProviders/FeaturePack/FeaturePackContentSource.cpp
--- a/Engine/Source/Editor/AddContentDialog/Private/ContentSourceProviders/FeaturePack/FeaturePackContentSource.cpp
+++ b/Engine/Source/Editor/AddContentDialog/Private/ContentSourceProviders/FeaturePack/FeaturePackContentSource.cpp
@@ -189,12 +189,15 @@ FFeaturePackContentSource::FFeaturePackContentSource(FString InFeaturePackPath)
 	}
 	
 	// Parse asset types field
-	for (TSharedPtr<FJsonValue> AssetTypesValue : ManifestObject->GetArrayField("SearchTags"))
+	if( ManifestObject->HasField("SearchTags")==true)
 	{
-		TSharedPtr<FJsonObject> LocalizedAssetTypesObject = AssetTypesValue->AsObject();
-		LocalizedSearchTags.Add(FLocalizedTextArray(
-			LocalizedAssetTypesObject->GetStringField("Language"),
-			LocalizedAssetTypesObject->GetStringField("Text")));
+		for (TSharedPtr<FJsonValue> AssetTypesValue : ManifestObject->GetArrayField("SearchTags"))
+		{
+			TSharedPtr<FJsonObject> LocalizedAssetTypesObject = AssetTypesValue->AsObject();
+			LocalizedSearchTags.Add(FLocalizedTextArray(
+				LocalizedAssetTypesObject->GetStringField("Language"),
+				LocalizedAssetTypesObject->GetStringField("Text")));
+		}
 	}
 
 	// Parse class types field
