Fixed incorrect check for focusasset

diff --git a/Engine/Source/Editor/AddContentDialog/Private/ContentSourceProviders/FeaturePack/FeaturePackContentSource.cpp b/Engine/Source/Editor/AddContentDialog/Private/ContentSourceProviders/FeaturePack/FeaturePackContentSource.cpp
--- a/Engine/Source/Editor/AddContentDialog/Private/ContentSourceProviders/FeaturePack/FeaturePackContentSource.cpp
+++ b/Engine/Source/Editor/AddContentDialog/Private/ContentSourceProviders/FeaturePack/FeaturePackContentSource.cpp
@@ -205,7 +205,7 @@ FFeaturePackContentSource::FFeaturePackContentSource(FString InFeaturePackPath)
 	
 	
 	// Parse initial focus asset if we have one - this is not required
-	if (ManifestObject->HasTypedField<EJson::String>("FocusAsset") == false)
+	if (ManifestObject->HasTypedField<EJson::String>("FocusAsset") == true)
 	{
 		FocusAssetIdent = ManifestObject->GetStringField("FocusAsset");
 	}	
