Added featrure pack automation test
Removed FeaturePack class and put functions from it in FeaturePackContentSource

diff --git a/Engine/Source/Editor/AddContentDialog/Private/ContentSourceProviders/FeaturePack/FeaturePackContentSource.cpp b/Engine/Source/Editor/AddContentDialog/Private/ContentSourceProviders/FeaturePack/FeaturePackContentSource.cpp
--- a/Engine/Source/Editor/AddContentDialog/Private/ContentSourceProviders/FeaturePack/FeaturePackContentSource.cpp
+++ b/Engine/Source/Editor/AddContentDialog/Private/ContentSourceProviders/FeaturePack/FeaturePackContentSource.cpp
@@ -121,7 +121,7 @@ bool TryValidateManifestObject(TSharedPtr<FJsonObject> ManifestObject, TSharedPt
 	return true;
 }
 
-FFeaturePackContentSource::FFeaturePackContentSource(FString InFeaturePackPath)
+FFeaturePackContentSource::FFeaturePackContentSource(FString InFeaturePackPath, bool bDontRegisterForSearch)
 {
 	FeaturePackPath = InFeaturePackPath;
 	bPackValid = false;
@@ -136,7 +136,7 @@ FFeaturePackContentSource::FFeaturePackContentSource(FString InFeaturePackPath)
 	TArray<uint8> ManifestBuffer;
 	if( LoadPakFileToBuffer(PakPlatformFile, FPaths::Combine(*MountPoint, TEXT("manifest.json")), ManifestBuffer) == false )
 	{
-		UE_LOG(LogFeaturePack, Warning, TEXT("Error in Feature pack %s. Cannot find manifest."), *InFeaturePackPath);
+		RecordAndLogError( FString::Printf(TEXT("Error in Feature pack %s. Cannot find manifest."), *InFeaturePackPath));
 		Category = EContentSourceCategory::Unknown;
 		return;
 	}
@@ -150,7 +150,7 @@ FFeaturePackContentSource::FFeaturePackContentSource(FString InFeaturePackPath)
 
 	if (ManifestReader->GetErrorMessage().IsEmpty() == false)
 	{
-		UE_LOG(LogFeaturePack, Warning, TEXT("Error in Feature pack %s. Failed to parse manifest: %s"), *InFeaturePackPath, *ManifestReader->GetErrorMessage());
+		RecordAndLogError( FString::Printf(TEXT("Error in Feature pack %s. Failed to parse manifest: %s"), *InFeaturePackPath, *ManifestReader->GetErrorMessage()));
 		Category = EContentSourceCategory::Unknown;
 		return;
 	}
@@ -158,7 +158,7 @@ FFeaturePackContentSource::FFeaturePackContentSource(FString InFeaturePackPath)
 	TSharedPtr<FString> ManifestObjectErrorMessage;
 	if (TryValidateManifestObject(ManifestObject, ManifestObjectErrorMessage) == false)
 	{
-		UE_LOG(LogFeaturePack, Warning, TEXT("Error in Feature pack %s. Manifest object error: %s"), *InFeaturePackPath, **ManifestObjectErrorMessage);
+		RecordAndLogError( FString::Printf(TEXT("Error in Feature pack %s. Manifest object error: %s"), *InFeaturePackPath, **ManifestObjectErrorMessage));
 		Category = EContentSourceCategory::Unknown;
 		return;
 	}
@@ -229,7 +229,7 @@ FFeaturePackContentSource::FFeaturePackContentSource(FString InFeaturePackPath)
 	}
 	else
 	{
-		UE_LOG(LogFeaturePack, Warning, TEXT("Error in Feature pack %s. Cannot find thumbnail %s."), *InFeaturePackPath, *ThumbnailFile );
+		RecordAndLogError( FString::Printf(TEXT("Error in Feature pack %s. Cannot find thumbnail %s."), *InFeaturePackPath, *ThumbnailFile ));
 	}
 
 	const TArray<TSharedPtr<FJsonValue>> ScreenshotFilenameArray = ManifestObject->GetArrayField("Screenshots");
@@ -242,13 +242,16 @@ FFeaturePackContentSource::FFeaturePackContentSource(FString InFeaturePackPath)
 		}
 		else
 		{
-			UE_LOG(LogFeaturePack, Warning, TEXT("Error in Feature pack %s. Cannot find screenshot %s."), *InFeaturePackPath, *ScreenshotFilename->AsString() );
+			RecordAndLogError( FString::Printf(TEXT("Error in Feature pack %s. Cannot find screenshot %s."), *InFeaturePackPath, *ScreenshotFilename->AsString() ));
 		}
 	}
 
-	FSuperSearchModule& SuperSearchModule = FModuleManager::LoadModuleChecked< FSuperSearchModule >(TEXT("SuperSearch"));
-	SuperSearchModule.GetActOnSearchTextClicked().AddRaw(this, &FFeaturePackContentSource::HandleActOnSearchText);	
-	SuperSearchModule.GetSearchTextChanged().AddRaw(this, &FFeaturePackContentSource::HandleSuperSearchTextChanged);
+	if( bDontRegisterForSearch == false )
+	{
+		FSuperSearchModule& SuperSearchModule = FModuleManager::LoadModuleChecked< FSuperSearchModule >(TEXT("SuperSearch"));
+		SuperSearchModule.GetActOnSearchTextClicked().AddRaw(this, &FFeaturePackContentSource::HandleActOnSearchText);
+		SuperSearchModule.GetSearchTextChanged().AddRaw(this, &FFeaturePackContentSource::HandleSuperSearchTextChanged);
+	}
 	bPackValid = true;
 }
 
@@ -468,4 +471,87 @@ FLocalizedTextArray FFeaturePackContentSource::ChooseLocalizedTextArray(TArray<F
 	return Default;
 }
 
+void FFeaturePackContentSource::ImportPendingPacks()
+{ 
+	bool bAddPacks;
+	if (GConfig->GetBool(TEXT("StartupActions"), TEXT("bAddPacks"), bAddPacks, GGameIni) == true)
+	{
+		if (bAddPacks == true)
+		{
+			ParseAndImportPacks();			
+			GConfig->SetBool(TEXT("StartupActions"), TEXT("bAddPacks"), false, GGameIni);
+			GConfig->Flush(true, GGameIni);
+		}
+	}
+}
+
+void FFeaturePackContentSource::ParseAndImportPacks()
+{
+	FAssetToolsModule& AssetToolsModule = FModuleManager::Get().LoadModuleChecked<FAssetToolsModule>("AssetTools");
+	// Look for pack insertions in the startup actions section
+	TArray<FString> PacksToAdd;
+	GConfig->GetArray(TEXT("StartupActions"), TEXT("InsertPack"), PacksToAdd, GGameIni);
+	int32 PacksInserted = 0;
+	for (int32 iPackEntry = 0; iPackEntry < PacksToAdd.Num(); iPackEntry++)
+	{
+		FPackData EachPackData;
+		TArray<FString> PackEntries;
+		PacksToAdd[iPackEntry].ParseIntoArray(PackEntries, TEXT(","), true);
+		FString PackSource;
+		FString PackName;
+		// Parse the pack name and source
+		for (int32 iEntry = 0; iEntry < PackEntries.Num(); iEntry++)
+		{
+			FString EachString = PackEntries[iEntry];
+			// remove the parenthesis
+			EachString = EachString.Replace(TEXT("("), TEXT(""));
+			EachString = EachString.Replace(TEXT(")"), TEXT(""));
+			if (EachString.StartsWith(TEXT("PackSource=")) == true)
+			{
+				EachString = EachString.Replace(TEXT("PackSource="), TEXT(""));
+				EachString = EachString.TrimQuotes();
+				EachPackData.PackSource = EachString;
+			}
+			if (EachString.StartsWith(TEXT("PackName=")) == true)
+			{
+				EachString = EachString.Replace(TEXT("PackName="), TEXT(""));
+				EachString = EachString.TrimQuotes();
+				EachPackData.PackName = EachString;
+			}
+		}
+
+		// If we found anything to insert, insert it !
+		if ((EachPackData.PackSource.IsEmpty() == false) && (EachPackData.PackName.IsEmpty() == false))
+		{
+			TArray<FString> EachImport;
+			FString FullPath = FPaths::FeaturePackDir() + EachPackData.PackSource;
+			EachImport.Add(FullPath);
+			EachPackData.ImportedObjects = AssetToolsModule.Get().ImportAssets(EachImport, TEXT("/Game"), nullptr, false);
+
+			if (EachPackData.ImportedObjects.Num() == 0)
+			{
+				UE_LOG(LogFeaturePack, Warning, TEXT("No objects imported installing pack %s"), *EachPackData.PackSource);
+			}
+			else
+			{
+				// Save any imported assets.
+				TArray<UPackage*> ToSave;
+				for (auto ImportedObject : EachPackData.ImportedObjects)
+				{
+					ToSave.AddUnique(ImportedObject->GetOutermost());
+				}
+				FEditorFileUtils::PromptForCheckoutAndSave(ToSave, /*bCheckDirty=*/ false, /*bPromptToSave=*/ false);
+				PacksInserted++;
+			}
+		}
+	}
+	UE_LOG(LogFeaturePack, Warning, TEXT("Inserted %d feature packs"), PacksInserted++);
+}
+
+void FFeaturePackContentSource::RecordAndLogError(const FString& ErrorString)
+{
+	UE_LOG(LogFeaturePack, Warning, TEXT("%s"), *ErrorString);
+	ParseErrors.Add(ErrorString);
+}
+
 #undef LOCTEXT_NAMESPACE 
diff --git a/Engine/Source/Editor/AddContentDialog/Private/ContentSourceProviders/FeaturePack/FeaturePackContentSource.h b/Engine/Source/Editor/AddContentDialog/Private/ContentSourceProviders/FeaturePack/FeaturePackContentSource.h
deleted file mode 100644
diff --git a/Engine/Source/Editor/AddContentDialog/Private/ContentSourceProviders/FeaturePack/FeaturePackContentSourceProvider.cpp b/Engine/Source/Editor/AddContentDialog/Private/ContentSourceProviders/FeaturePack/FeaturePackContentSourceProvider.cpp
--- a/Engine/Source/Editor/AddContentDialog/Private/ContentSourceProviders/FeaturePack/FeaturePackContentSourceProvider.cpp
+++ b/Engine/Source/Editor/AddContentDialog/Private/ContentSourceProviders/FeaturePack/FeaturePackContentSourceProvider.cpp
@@ -29,7 +29,7 @@ public:
 
 FFeaturePackContentSourceProvider::FFeaturePackContentSourceProvider()
 {
-	FeaturePackPath = FPaths::Combine(*FPaths::RootDir(), TEXT("FeaturePacks"));
+	FeaturePackPath = FPaths::FeaturePackDir();
 	StartUpDirectoryWatcher();
 	RefreshFeaturePacks();
 }
diff --git a/Engine/Source/Editor/AddContentDialog/Public/FeaturePackContentSource.h b/Engine/Source/Editor/AddContentDialog/Public/FeaturePackContentSource.h
new file mode 100644
--- /dev/null
+++ b/Engine/Source/Editor/AddContentDialog/Public/FeaturePackContentSource.h
@@ -0,0 +1,110 @@
+// Copyright 1998-2015 Epic Games, Inc. All Rights Reserved.
+
+#pragma once
+#include "Private/IContentSource.h"
+
+class UObject;
+
+struct FPackData
+{
+	FString PackSource;
+	FString PackName;
+	FString PackMap;
+	TArray<UObject*>	ImportedObjects;
+};
+
+class FLocalizedTextArray
+{
+public:
+	FLocalizedTextArray()
+	{
+	}
+
+	/** Creates a new FLocalizedText
+		@param InTwoLetterLanguage - The iso 2-letter language specifier.
+		@param InText - The text in the language specified */
+	FLocalizedTextArray(FString InTwoLetterLanguage, FString InText)
+	{
+		TwoLetterLanguage = InTwoLetterLanguage;
+		TArray<FString> AsArray;
+		InText.ParseIntoArray(AsArray,TEXT(","));
+		for (int32 iString = 0; iString < AsArray.Num() ; iString++)
+		{
+			Tags.Add(FText::FromString(AsArray[iString]));
+		}
+	}
+
+	/** Gets the iso 2-letter language specifier for this text. */
+	FString GetTwoLetterLanguage() const
+	{
+		return TwoLetterLanguage;
+	}
+
+	/** Gets the array of tags in the language specified. */
+	TArray<FText> GetTags() const
+	{
+		return Tags;
+	}
+
+private:
+	FString TwoLetterLanguage;
+	TArray<FText> Tags;
+};
+
+
+class FPakPlatformFile;
+struct FSearchEntry;
+
+/** A content source which represents a content upack. */
+class FFeaturePackContentSource : public IContentSource
+{
+public:
+	ADDCONTENTDIALOG_API FFeaturePackContentSource(FString InFeaturePackPath, bool bDontRegister = false);
+
+	virtual TArray<FLocalizedText> GetLocalizedNames() const override;
+	virtual TArray<FLocalizedText> GetLocalizedDescriptions() const override;
+	
+	virtual EContentSourceCategory GetCategory() const override;
+	virtual TArray<FLocalizedText> GetLocalizedAssetTypes() const override;
+	virtual FString GetSortKey() const override;
+	virtual FString GetClassTypesUsed() const override;
+	
+	virtual TSharedPtr<FImageData> GetIconData() const override;
+	virtual TArray<TSharedPtr<FImageData>> GetScreenshotData() const override;
+
+	virtual bool InstallToProject(FString InstallPath) override;
+	virtual bool IsDataValid() const override;
+	
+	virtual ~FFeaturePackContentSource();
+	
+	void HandleActOnSearchText(TSharedPtr<FSearchEntry> SearchEntry);
+	void HandleSuperSearchTextChanged(const FString& InText, TArray< TSharedPtr<FSearchEntry> >& OutSuggestions);
+
+	static ADDCONTENTDIALOG_API void ImportPendingPacks();
+	TArray<FString>	ParseErrors;
+private:
+	static void ParseAndImportPacks();
+	bool LoadPakFileToBuffer(FPakPlatformFile& PakPlatformFile, FString Path, TArray<uint8>& Buffer);
+	FString GetFocusAssetName() const;
+
+private:
+	void RecordAndLogError(const FString& ErrorString);
+	/** Selects an FLocalizedText from an array which matches either the supplied language code, or the default language code. */
+	FLocalizedTextArray ChooseLocalizedTextArray(TArray<FLocalizedTextArray> Choices, FString LanguageCode);
+	FLocalizedText ChooseLocalizedText(TArray<FLocalizedText> Choices, FString LanguageCode);
+	void TryAddFeaturePackCategory(FString CategoryTitle, TArray< TSharedPtr<FSearchEntry> >& OutSuggestions);
+
+	FString FeaturePackPath;
+	TArray<FLocalizedText> LocalizedNames;
+	TArray<FLocalizedText> LocalizedDescriptions;
+	EContentSourceCategory Category;
+	TSharedPtr<FImageData> IconData;
+	TArray<TSharedPtr<FImageData>> ScreenshotData;
+	TArray<FLocalizedText> LocalizedAssetTypesList;
+	FString ClassTypes;
+	bool bPackValid;
+	FString FocusAssetIdent;
+	FString SortKey;
+	TArray<FLocalizedTextArray> LocalizedSearchTags;
+	
+};
\ No newline at end of file
diff --git a/Engine/Source/Editor/UnrealEd/Private/FeaturePack/FeaturePack.cpp b/Engine/Source/Editor/UnrealEd/Private/FeaturePack/FeaturePack.cpp
deleted file mode 100644
diff --git a/Engine/Source/Editor/UnrealEd/Private/Tests/FeaturePackAutomationTests.cpp b/Engine/Source/Editor/UnrealEd/Private/Tests/FeaturePackAutomationTests.cpp
new file mode 100644
--- /dev/null
+++ b/Engine/Source/Editor/UnrealEd/Private/Tests/FeaturePackAutomationTests.cpp
@@ -0,0 +1,61 @@
+// Copyright 1998-2015 Epic Games, Inc. All Rights Reserved.
+
+#include "UnrealEd.h"
+
+#include "Tests/AutomationTestSettings.h"
+
+#if WITH_EDITOR
+#include "FileHelpers.h"
+#include "LevelEditor.h"
+#include "AutomationEditorCommon.h"
+#endif
+
+#include "AssetSelection.h"
+#include "AutomationCommon.h"
+#include "AutomationTest.h"
+#include "FeaturePackContentSource.h"
+
+DEFINE_LOG_CATEGORY_STATIC(FeaturePackAutomationTestLog, Log, All);
+
+// 
+// namespace FeaturePackAutomationTests
+// {
+// }
+
+/**
+ * Feature Pack validity test
+ */
+IMPLEMENT_SIMPLE_AUTOMATION_TEST( FFeaturePackValidityTest, "System.Editor.Content.FeaturePackValidityTest", EAutomationTestFlags::ATF_Editor )
+
+/** 
+ * Scans feature pack folder and verifies that all feature packs therein are parsed successfully.
+ *
+ * @param Parameters - Unused for this test
+ * @return	TRUE if the test was successful, FALSE otherwise
+ */
+bool FFeaturePackValidityTest::RunTest(const FString& Parameters)
+{
+	FString FeaturePackPath = FPaths::FeaturePackDir();
+	FeaturePackPath += TEXT("*.upack");
+	int32 FailCount = 0;
+	IPlatformFile &PlatformFile = FPlatformFileManager::Get().GetPlatformFile();
+	
+	TArray<FString> FeaturePacks;
+	IFileManager::Get().FindFiles(FeaturePacks, *FeaturePackPath, true, false );
+
+	for (auto FeaturePackFile : FeaturePacks)
+	{
+		TUniquePtr<FFeaturePackContentSource> NewContentSource = MakeUnique<FFeaturePackContentSource>(FPaths::FeaturePackDir() + FeaturePackFile, true);
+		if (NewContentSource->IsDataValid() == false)
+		{
+			FailCount++;
+			for (int32 iError = 0; iError < NewContentSource->ParseErrors.Num() ; iError++)
+			{
+				AddError(NewContentSource->ParseErrors[iError]);
+			}
+			
+		}
+	}
+	
+	return FailCount==0;
+}
diff --git a/Engine/Source/Editor/UnrealEd/Private/UnrealEdMisc.cpp b/Engine/Source/Editor/UnrealEd/Private/UnrealEdMisc.cpp
--- a/Engine/Source/Editor/UnrealEd/Private/UnrealEdMisc.cpp
+++ b/Engine/Source/Editor/UnrealEd/Private/UnrealEdMisc.cpp
@@ -37,11 +37,11 @@
 #include "HotReloadInterface.h"
 #include "PerformanceMonitor.h"
 #include "Engine/WorldComposition.h"
-#include "FeaturePack.h"
 #include "GameMapsSettings.h"
 #include "GeneralProjectSettings.h"
 #include "Lightmass/LightmappedSurfaceCollection.h"
 #include "IProjectManager.h"
+#include "FeaturePackContentSource.h"
 
 #define LOCTEXT_NAMESPACE "UnrealEd"
 
@@ -260,8 +260,7 @@ void FUnrealEdMisc::OnInit()
 		bool bMapLoaded = false;
 
 		// Insert any feature packs if required. We need to do this before we try and load a map since any pack may contain a map
-		FFeaturePack FeaturePackHandler;
-		FeaturePackHandler.ImportPendingPacks();
+		FFeaturePackContentSource::ImportPendingPacks();
 
 		FString ParsedMapName;
 		if ( FParse::Token(ParsedCmdLine, ParsedMapName, false) && 
diff --git a/Engine/Source/Editor/UnrealEd/Public/FeaturePack/FeaturePack.h b/Engine/Source/Editor/UnrealEd/Public/FeaturePack/FeaturePack.h
deleted file mode 100644
diff --git a/Engine/Source/Editor/UnrealEd/UnrealEd.Build.cs b/Engine/Source/Editor/UnrealEd/UnrealEd.Build.cs
--- a/Engine/Source/Editor/UnrealEd/UnrealEd.Build.cs
+++ b/Engine/Source/Editor/UnrealEd/UnrealEd.Build.cs
@@ -14,7 +14,6 @@ public class UnrealEd : ModuleRules
 			{
 				"Editor/UnrealEd/Private",
 				"Editor/UnrealEd/Private/Settings",
-				"Editor/UnrealEd/Private/FeaturePack",
 				"Editor/PackagesDialog/Public",				
 				"Developer/DerivedDataCache/Public",
 				"Developer/TargetPlatform/Public",
@@ -132,6 +131,7 @@ public class UnrealEd : ModuleRules
 				"TreeMap",
                 "MaterialUtilities",
                 "LocalizationService",
+				"AddContentDialog",
 			}
 		);
 
@@ -233,6 +233,7 @@ public class UnrealEd : ModuleRules
 				"NiagaraEditor",
                 "Niagara",
                 "VectorVM",
+				"AddContentDialog",
 			}
 			);
 
