Code to support shared template resources

diff --git a/Engine/Source/Editor/AddContentDialog/Private/ContentSourceProviders/FeaturePack/FeaturePackContentSource.cpp b/Engine/Source/Editor/AddContentDialog/Private/ContentSourceProviders/FeaturePack/FeaturePackContentSource.cpp
--- a/Engine/Source/Editor/AddContentDialog/Private/ContentSourceProviders/FeaturePack/FeaturePackContentSource.cpp
+++ b/Engine/Source/Editor/AddContentDialog/Private/ContentSourceProviders/FeaturePack/FeaturePackContentSource.cpp
@@ -127,7 +127,6 @@ FFeaturePackContentSource::FFeaturePackContentSource(FString InFeaturePackPath,
 	bPackValid = false;
 	// Create a pak platform file and mount the feature pack file.
 	FPakPlatformFile PakPlatformFile;
-	FString CommandLine;
 	PakPlatformFile.Initialize(&FPlatformFileManager::Get().GetPlatformFile(), TEXT(""));
 	FString MountPoint = "root:/";
 	PakPlatformFile.Mount(*InFeaturePackPath, 0, *MountPoint);
@@ -188,7 +187,7 @@ FFeaturePackContentSource::FFeaturePackContentSource(FString InFeaturePackPath,
 			FText::FromString(LocalizedAssetTypesObject->GetStringField("Text"))));
 	}
 	
-	// Parse asset types field
+	// Parse search tags field
 	if( ManifestObject->HasField("SearchTags")==true)
 	{
 		for (TSharedPtr<FJsonValue> AssetTypesValue : ManifestObject->GetArrayField("SearchTags"))
@@ -201,8 +200,7 @@ FFeaturePackContentSource::FFeaturePackContentSource(FString InFeaturePackPath,
 	}
 
 	// Parse class types field
-	ClassTypes = ManifestObject->GetStringField("ClassTypes");
-	
+	ClassTypes = ManifestObject->GetStringField("ClassTypes");	
 	
 	// Parse initial focus asset if we have one - this is not required
 	if (ManifestObject->HasTypedField<EJson::String>("FocusAsset") == true)
@@ -232,6 +230,7 @@ FFeaturePackContentSource::FFeaturePackContentSource(FString InFeaturePackPath,
 		RecordAndLogError( FString::Printf(TEXT("Error in Feature pack %s. Cannot find thumbnail %s."), *InFeaturePackPath, *ThumbnailFile ));
 	}
 
+	// parse the screenshots field
 	const TArray<TSharedPtr<FJsonValue>> ScreenshotFilenameArray = ManifestObject->GetArrayField("Screenshots");
 	for (const TSharedPtr<FJsonValue> ScreenshotFilename : ScreenshotFilenameArray)
 	{
@@ -246,6 +245,39 @@ FFeaturePackContentSource::FFeaturePackContentSource(FString InFeaturePackPath,
 		}
 	}
 
+	// Parse additional packs data
+	if (ManifestObject->HasTypedField<EJson::Array>("AdditionalFeaturePacks") == true)
+	{		
+		UEnum* DetailEnum = FindObjectChecked<UEnum>(ANY_PACKAGE, TEXT("EFeaturePackDetailLevel"));
+		for (TSharedPtr<FJsonValue> AdditionalFeaturePackValue : ManifestObject->GetArrayField("AdditionalFeaturePacks"))
+		{
+			TSharedPtr<FJsonObject> EachAdditionalPack = AdditionalFeaturePackValue->AsObject();
+			FString MountName = EachAdditionalPack->GetStringField("MountName");
+
+			TArray<EFeaturePackDetailLevel>	Levels;
+			for (TSharedPtr<FJsonValue> DetailValue: EachAdditionalPack->GetArrayField("DetailLevels"))
+			{
+				const FString DetailString = DetailValue->AsString();
+				int32 Index = DetailEnum->FindEnumIndex(FName(*DetailString));
+				EFeaturePackDetailLevel EachLevel = Index != INDEX_NONE ? (EFeaturePackDetailLevel)Index : EFeaturePackDetailLevel::Standard;
+				Levels.AddUnique(EachLevel);
+			}
+			int32 PackIndex = AdditionalFeaturePacks.Add(FFeaturePackLevelSet(
+				MountName,
+				Levels));
+			// Check that the additional feature packs listed exist
+			FFeaturePackLevelSet& NewSet = AdditionalFeaturePacks[PackIndex];
+			for (int32 iLevel = 0; iLevel < NewSet.DetailLevels.Num() ; iLevel++)
+			{
+				FString FullPath = FPaths::FeaturePackDir() + NewSet.GetFeaturePackNameForLevel(NewSet.DetailLevels[iLevel]);	
+				if (!FPlatformFileManager::Get().GetPlatformFile().FileExists(*FullPath))
+				{
+					RecordAndLogError( FString::Printf(TEXT("Error in Feature pack %s. Cannot find additional pack %s."), *FullPath));
+				}
+			}		
+		}
+	}
+
 	if( bDontRegisterForSearch == false )
 	{
 		FSuperSearchModule& SuperSearchModule = FModuleManager::LoadModuleChecked< FSuperSearchModule >(TEXT("SuperSearch"));
@@ -330,6 +362,27 @@ bool FFeaturePackContentSource::InstallToProject(FString InstallPath)
 				ToSave.AddUnique(ImportedObject->GetOutermost());
 			}
 			FEditorFileUtils::PromptForCheckoutAndSave( ToSave, /*bCheckDirty=*/ false, /*bPromptToSave=*/ false );
+
+			// Now copy files in any additional packs (These are for the shared assets)
+ 			EFeaturePackDetailLevel RequiredLevel = EFeaturePackDetailLevel::High;
+ 			for (int32 iExtraPack = 0; iExtraPack < AdditionalFeaturePacks.Num(); iExtraPack++)
+ 			{
+				FString DestinationFolder = FPaths::GameDir();
+				FString FullPath = FPaths::FeaturePackDir() + AdditionalFeaturePacks[iExtraPack].GetFeaturePackNameForLevel(RequiredLevel);
+ 				if (FPlatformFileManager::Get().GetPlatformFile().FileExists(*FullPath))
+ 				{
+ 					FString InsertPath = DestinationFolder + *AdditionalFeaturePacks[iExtraPack].MountName;
+
+					TUniquePtr<FFeaturePackContentSource> NewContentSource = MakeUnique<FFeaturePackContentSource>(FullPath, true);
+					if (NewContentSource->IsDataValid() == true)
+					{
+						bool bHasSourceFiles = false;
+						TArray<FString> FilesCopied;
+						NewContentSource->CopyAdditionalFilesToFolder(DestinationFolder, FilesCopied, bHasSourceFiles);
+					}
+ 				}
+ 			}
+
 			bResult = true;
 			
 			// Focus on a specific asset if we want to.
@@ -471,6 +524,88 @@ FLocalizedTextArray FFeaturePackContentSource::ChooseLocalizedTextArray(TArray<F
 	return Default;
 }
 
+bool FFeaturePackContentSource::GetAdditionalFilesForPack(TArray<FString>& FileList, bool& bContainsSource)
+{
+	bool bParsedFiles = false;
+	if (bPackValid == false)
+	{
+		RecordAndLogError(FString::Printf(TEXT("Cannot extract files from invalid Pack %s"), *FeaturePackPath));
+	}
+	else
+	{
+		// Create a pak platform file and mount the feature pack file.
+		FPakPlatformFile PakPlatformFile;
+		PakPlatformFile.Initialize(&FPlatformFileManager::Get().GetPlatformFile(), TEXT(""));
+		FString MountPoint = "root:/";
+		PakPlatformFile.Mount(*FeaturePackPath, 0, *MountPoint);
+
+		// Gets the manifest file as a JSon string
+		TArray<uint8> ManifestBuffer;
+		if (LoadPakFileToBuffer(PakPlatformFile, FPaths::Combine(*MountPoint, TEXT("Config/Config.ini")), ManifestBuffer) == false)
+		{
+			RecordAndLogError(FString::Printf(TEXT("Error in Feature pack %s. Cannot find Config.ini"), *FeaturePackPath));			
+		}
+		else
+		{
+			FString ConfigFileString;
+			FFileHelper::BufferToString(ConfigFileString, ManifestBuffer.GetData(), ManifestBuffer.Num());
+			bParsedFiles = ExtractListOfAdditionalFiles(ConfigFileString, FileList, bContainsSource);
+		}
+	}
+	return bParsedFiles;
+}
+
+bool FFeaturePackContentSource::ExtractListOfAdditionalFiles(const FString& ConfigFileAsString,TArray<FString>& FileList, bool& bContainsSource) const
+{
+	FConfigFile PackConfig;
+	PackConfig.ProcessInputFileContents(ConfigFileAsString);
+	FConfigSection* AdditionalFilesSection = PackConfig.Find("AdditionalFilesToAdd");
+	
+	bContainsSource = false;
+	bool bParsedAdditionFiles = false;
+	if (AdditionalFilesSection)
+	{
+		bParsedAdditionFiles = true;
+		for (auto FilePair : *AdditionalFilesSection)
+		{
+			if (FilePair.Key.ToString().Contains("Files"))
+			{
+				FString Filename = FPaths::GetCleanFilename(FilePair.Value);
+				FString Directory = FPaths::RootDir() / FPaths::GetPath(FilePair.Value);
+				FPaths::MakeStandardFilename(Directory);
+				FPakFile::MakeDirectoryFromPath(Directory);
+
+				if (Filename.Contains(TEXT("*")))
+				{
+					TArray<FString> FoundFiles;
+					IFileManager::Get().FindFilesRecursive(FoundFiles, *Directory, *Filename, true, false);
+					FileList.Append(FoundFiles);
+					if (!bContainsSource)
+					{
+						for (const FString& FoundFile : FoundFiles)
+						{
+							if (FoundFile.StartsWith(TEXT("Source/")) || FoundFile.Contains(TEXT("/Source/")))
+							{
+								bContainsSource = true;
+								break;
+							}
+						}
+					}
+				}
+				else
+				{
+					FileList.Add(Directory / Filename);
+					if (!bContainsSource && (FileList.Last().StartsWith(TEXT("Source/")) || FileList.Last().Contains(TEXT("/Source/"))))
+					{
+						bContainsSource = true;
+					}
+				}
+			}
+		}
+	}
+	return bParsedAdditionFiles;
+}
+
 void FFeaturePackContentSource::ImportPendingPacks()
 { 
 	bool bAddPacks;
@@ -554,4 +689,35 @@ void FFeaturePackContentSource::RecordAndLogError(const FString& ErrorString)
 	ParseErrors.Add(ErrorString);
 }
 
+void FFeaturePackContentSource::CopyAdditionalFilesToFolder( const FString& DestinationFolder, TArray<FString>& FilesCopied, bool &bHasSourceFiles )
+{
+	FString ContentIdent = TEXT("Content/");
+	TArray<FString> FilesToAdd;
+	GetAdditionalFilesForPack(FilesToAdd, bHasSourceFiles);
+	for (int32 iFile = 0; iFile < FilesToAdd.Num(); ++iFile)
+	{
+		FString EachFile = FilesToAdd[iFile];
+		int32 ContentIndex;
+		ContentIndex = EachFile.Find(*ContentIdent);
+		if (ContentIndex != INDEX_NONE)
+		{
+			FString ContentFile = EachFile.RightChop(ContentIndex);
+			FString GameFolder = EachFile.Left(ContentIndex - 1);
+			int32 NameStart = GameFolder.Find(TEXT("/"), ESearchCase::IgnoreCase, ESearchDir::FromEnd);
+			GameFolder = GameFolder.RightChop(NameStart);
+			FPaths::NormalizeFilename(ContentFile);
+			ContentFile.InsertAt(ContentIdent.Len() - 1, GameFolder);
+			FString FinalDestination = DestinationFolder / ContentFile;
+			if (IFileManager::Get().Copy(*FinalDestination, *EachFile) == COPY_OK)
+			{
+				FilesCopied.Add(FinalDestination);
+			}
+			else
+			{
+				UE_LOG(LogFeaturePack, Warning, TEXT("Failed to copy %s to %s"), *EachFile, *FinalDestination);
+			}
+		}
+	}
+}
+
 #undef LOCTEXT_NAMESPACE 
diff --git a/Engine/Source/Editor/AddContentDialog/Public/FeaturePackContentSource.h b/Engine/Source/Editor/AddContentDialog/Public/FeaturePackContentSource.h
--- a/Engine/Source/Editor/AddContentDialog/Public/FeaturePackContentSource.h
+++ b/Engine/Source/Editor/AddContentDialog/Public/FeaturePackContentSource.h
@@ -2,6 +2,7 @@
 
 #pragma once
 #include "Private/IContentSource.h"
+#include "FeaturePackContentSource.generated.h"
 
 class UObject;
 
@@ -52,6 +53,59 @@ private:
 };
 
 
+/** Defines categories for shared template resource levels. */
+UENUM()
+enum class EFeaturePackDetailLevel :uint8
+{
+	Standard,
+	High,
+};
+
+/* Structure that defines a shared feature pack resource. */
+USTRUCT()
+struct FFeaturePackLevelSet
+{
+	GENERATED_BODY()
+
+	FFeaturePackLevelSet(){};
+
+	/** Creates a new FFeaturePackLevelSet
+		@param InMountName - Name of the pack/folder to insert to 
+		@param InDetailLevels - The levels available for this pack*/
+	FFeaturePackLevelSet(FString InMountName, TArray<EFeaturePackDetailLevel> InDetailLevels)
+	{
+		MountName = InMountName;
+		DetailLevels = InDetailLevels;
+	}
+
+	/* List of shared resource levels for this shared resource.*/
+	UPROPERTY()
+	TArray<EFeaturePackDetailLevel> DetailLevels;
+
+	/* Mount name for the shared resource - this is the folder the resource will be copied to on project generation as well as the name of the folder that will appear in the content browser. */
+	UPROPERTY()
+	FString MountName;
+
+	FString GetFeaturePackNameForLevel(EFeaturePackDetailLevel InLevel, bool bLevelRequired = false)
+	{
+		check(DetailLevels.Num()>0); // We need at least one detail level defined
+		int32 Index = DetailLevels.Find(InLevel);
+		FString DetailString;
+		if( Index != INDEX_NONE)
+		{			
+			UEnum::GetValueAsString(TEXT("/Script/AddContentDialog.EFeaturePackDetailLevel"), InLevel, DetailString);					
+		}
+		else 
+		{
+			check(bLevelRequired==false); // The level is REQUIRED and we don't have it !
+			// If we didn't have the requested level, use the first
+			UEnum::GetValueAsString(TEXT("/Script/AddContentDialog.EFeaturePackDetailLevel"), DetailLevels[0], DetailString);
+		}
+		FString NameString = MountName+DetailString + TEXT(".upack");
+		return NameString;
+	}
+};
+
 class FPakPlatformFile;
 struct FSearchEntry;
 
@@ -79,15 +133,44 @@ public:
 	
 	void HandleActOnSearchText(TSharedPtr<FSearchEntry> SearchEntry);
 	void HandleSuperSearchTextChanged(const FString& InText, TArray< TSharedPtr<FSearchEntry> >& OutSuggestions);
+		
+	/*
+	 * Copies the list of files specified in 'AdditionFilesToInclude' section in the config.ini of the feature pack.
+	 *
+	 * @param DestinationFolder	Destination folder for the files
+	 * @param FilesCopied		List of files copied
+	 * @param bContainsSource 	Set to true if the file list contains any source files
+	 * @returns true if config file was read and parsed successfully
+	 */
+	ADDCONTENTDIALOG_API void CopyAdditionalFilesToFolder( const FString& DestinationFolder, TArray<FString>& FilesCopied, bool &bHasSourceFiles );
+
+	/*
+	 * Returns a list of additional files (including the path) as specified in the config file if one exists in the pack file.
+	 *
+	 * @param FileList		  array to receive list of files
+	 * @param bContainsSource did the file list contain any source files
+	 * @returns true if config file was read and parsed successfully
+	 */
+	ADDCONTENTDIALOG_API bool GetAdditionalFilesForPack(TArray<FString>& FileList, bool& bContainsSource);
 
 	static ADDCONTENTDIALOG_API void ImportPendingPacks();
+	
 	TArray<FString>	ParseErrors;
+	
 private:
 	static void ParseAndImportPacks();
 	bool LoadPakFileToBuffer(FPakPlatformFile& PakPlatformFile, FString Path, TArray<uint8>& Buffer);
 	FString GetFocusAssetName() const;
+	
+	/*
+	 * Extract the list of additional files defined in config file to an array
+	 *
+	 * @param ConfigFile	  config file as a string
+	 * @param FileList		  array to receive list of files
+	 * @param bContainsSource did the file list contain any source files
+	 */
+	bool ExtractListOfAdditionalFiles(const FString& ConfigFile, TArray<FString>& FileList,bool& bContainsSource) const;
 
-private:
 	void RecordAndLogError(const FString& ErrorString);
 	/** Selects an FLocalizedText from an array which matches either the supplied language code, or the default language code. */
 	FLocalizedTextArray ChooseLocalizedTextArray(TArray<FLocalizedTextArray> Choices, FString LanguageCode);
@@ -106,5 +189,5 @@ private:
 	FString FocusAssetIdent;
 	FString SortKey;
 	TArray<FLocalizedTextArray> LocalizedSearchTags;
-	
+	TArray<FFeaturePackLevelSet> AdditionalFeaturePacks;
 };
\ No newline at end of file
diff --git a/Engine/Source/Editor/GameProjectGeneration/Classes/TemplateProjectDefs.h b/Engine/Source/Editor/GameProjectGeneration/Classes/TemplateProjectDefs.h
--- a/Engine/Source/Editor/GameProjectGeneration/Classes/TemplateProjectDefs.h
+++ b/Engine/Source/Editor/GameProjectGeneration/Classes/TemplateProjectDefs.h
@@ -2,6 +2,7 @@
 
 #pragma once
 
+#include "FeaturePackContentSource.h"
 #include "TemplateProjectDefs.generated.h"
 
 // does not require reflection exposure
@@ -103,6 +104,10 @@ class UTemplateProjectDefs : public UObject
 	UPROPERTY(config)
 	TArray<FString> PacksToInclude;
 
+	/* Shared feature packs. The files in these packs listed in these structures marked as 'additionalfiles' will be copied on project generation*/
+	UPROPERTY(config)
+	TArray<FFeaturePackLevelSet> SharedContentPacks;
+
 	/** Fixes up all strings in this definitions object to replace \%TEMPLATENAME\% with the supplied template name and \%PROJECTNAME\% with the supplied project name */
 	void FixupStrings(const FString& TemplateName, const FString& ProjectName);
 
diff --git a/Engine/Source/Editor/GameProjectGeneration/GameProjectGeneration.Build.cs b/Engine/Source/Editor/GameProjectGeneration/GameProjectGeneration.Build.cs
--- a/Engine/Source/Editor/GameProjectGeneration/GameProjectGeneration.Build.cs
+++ b/Engine/Source/Editor/GameProjectGeneration/GameProjectGeneration.Build.cs
@@ -8,6 +8,11 @@ public class GameProjectGeneration : ModuleRules
 	{
         PrivateIncludePaths.AddRange(new string[] { "GameProjectGeneration/Private", "GameProjectGeneration/Public", "GameProjectGeneration/Classes" });
 
+		PublicDependencyModuleNames.AddRange(
+			new string[] {
+				"HardwareTargeting",
+			}
+		);
 
         PrivateIncludePathModuleNames.AddRange(
             new string[] {
@@ -15,6 +20,8 @@ public class GameProjectGeneration : ModuleRules
 				"ContentBrowser",
                 "DesktopPlatform",
                 "MainFrame",
+				"AddContentDialog",
+				"HardwareTargeting",
             }
 		);
 
@@ -38,6 +45,7 @@ public class GameProjectGeneration : ModuleRules
 				"UnrealEd",
 				"DesktopPlatform",
                 "HardwareTargeting",
+				"AddContentDialog",
 			}
 		);
 
diff --git a/Engine/Source/Editor/GameProjectGeneration/Private/GameProjectUtils.cpp b/Engine/Source/Editor/GameProjectGeneration/Private/GameProjectUtils.cpp
--- a/Engine/Source/Editor/GameProjectGeneration/Private/GameProjectUtils.cpp
+++ b/Engine/Source/Editor/GameProjectGeneration/Private/GameProjectUtils.cpp
@@ -27,7 +27,7 @@
 #include "HotReloadInterface.h"
 #include "SVerbChoiceDialog.h"
 #include "SourceCodeNavigation.h"
-///#include "AssetToolsModule.h"
+#include "FeaturePackContentSource.h"
 
 #include "SOutputLogDialog.h"
 
@@ -44,6 +44,8 @@ static_assert(PLATFORM_MAX_FILEPATH_LENGTH - MAX_PROJECT_PATH_BUFFER_SPACE > 0,
 TWeakPtr<SNotificationItem> GameProjectUtils::UpdateGameProjectNotification = NULL;
 TWeakPtr<SNotificationItem> GameProjectUtils::WarningProjectNameNotification = NULL;
 
+FString GameProjectUtils::DefaultFeaturePackExtension(TEXT(".upack"));	
+
 FText FNewClassInfo::GetClassName() const
 {
 	switch(ClassType)
@@ -603,7 +605,9 @@ void GameProjectUtils::GetStarterContentFiles(TArray<FString>& OutFilenames)
 {
 	FString const SrcFolder = FPaths::FeaturePackDir();
 	
-	IFileManager::Get().FindFilesRecursive(OutFilenames, *SrcFolder, TEXT("*.upack"), /*Files=*/true, /*Directories=*/false);
+	FString SearchPath = TEXT("*");
+	SearchPath += DefaultFeaturePackExtension;
+	IFileManager::Get().FindFilesRecursive(OutFilenames, *SrcFolder, *SearchPath, /*Files=*/true, /*Directories=*/false);
 }
 
 bool GameProjectUtils::CreateProject(const FProjectInformation& InProjectInfo, FText& OutFailReason, FText& OutFailLog)
@@ -1514,6 +1518,13 @@ bool GameProjectUtils::CreateProjectFromTemplate(const FProjectInformation& InPr
 		DeleteCreatedFiles(DestFolder, CreatedFiles);
 		return false;
 	}
+
+	if( AddSharedContentToProject(InProjectInfo, CreatedFiles, OutFailReason ) == false )
+	{
+		DeleteCreatedFiles(DestFolder, CreatedFiles);
+		return false;
+	}
+
 	
 	SlowTask.EnterProgressFrame();
 
@@ -3422,25 +3433,15 @@ bool GameProjectUtils::InsertFeaturePacksIntoINIFile(const FProjectInformation&
 		FString StarterPack;
 		if (InProjectInfo.TargetedHardware == EHardwareClass::Mobile)
 		{
-			StarterPack = TEXT("InsertPack=(PackSource=\"MobileStarterContent.upack\",PackName=\"StarterContent\")");
+			StarterPack = TEXT("InsertPack=(PackSource=\"MobileStarterContent") + DefaultFeaturePackExtension + TEXT(",PackName=\"StarterContent\")");
 		}
 		else
 		{
-			StarterPack = TEXT("InsertPack=(PackSource=\"StarterContent.upack\",PackName=\"StarterContent\")");
+			StarterPack = TEXT("InsertPack=(PackSource=\"StarterContent")  + DefaultFeaturePackExtension + TEXT(",PackName=\"StarterContent\")");
 		}
 		PackList.Add(StarterPack);
 	}
-
-	// Now any packs specified in the template def.
-	UTemplateProjectDefs* TemplateDefs = LoadTemplateDefs(SrcFolder);
-	if (TemplateDefs != NULL)
-	{
-		for (int32 iPack = 0; iPack < TemplateDefs->PacksToInclude.Num(); ++iPack)
-		{
-			PackList.Add(TemplateDefs->PacksToInclude[iPack]);
-		}
-	}
-
+	
 	if (PackList.Num() != 0)
 	{
 		FString FileOutput;
@@ -3470,4 +3471,81 @@ bool GameProjectUtils::InsertFeaturePacksIntoINIFile(const FProjectInformation&
 	return true;
 }
 
+bool GameProjectUtils::AddSharedContentToProject(const FProjectInformation &InProjectInfo, TArray<FString> &CreatedFiles, FText& OutFailReason)
+{
+	//const FString TemplateName = FPaths::GetBaseFilename(InProjectInfo.TemplateFile);
+	const FString SrcFolder = FPaths::GetPath(InProjectInfo.TemplateFile);
+	const FString DestFolder = FPaths::GetPath(InProjectInfo.ProjectFilename);
+
+	const FString ProjectConfigPath = DestFolder / TEXT("Config");
+	const FString IniFilename = ProjectConfigPath / TEXT("DefaultGame.ini");
+	
+	// Now any packs specified in the template def.
+	UTemplateProjectDefs* TemplateDefs = LoadTemplateDefs(SrcFolder);
+	if (TemplateDefs != NULL)
+	{
+		EFeaturePackDetailLevel RequiredDetail = EFeaturePackDetailLevel::High;
+		if (InProjectInfo.TargetedHardware == EHardwareClass::Mobile)
+		{
+			RequiredDetail = EFeaturePackDetailLevel::Standard;
+		}
+		
+		for (int32 iPack = 0; iPack < TemplateDefs->SharedContentPacks.Num(); ++iPack)
+		{
+			FFeaturePackLevelSet EachPack = TemplateDefs->SharedContentPacks[iPack];
+			EFeaturePackDetailLevel EachRequiredDetail = RequiredDetail;
+			
+			if (EachPack.DetailLevels.Num() == 1)
+			{
+				// If theres only only detail level override the requirement with that
+				EachRequiredDetail = EachPack.DetailLevels[0];
+			}
+			else if (EachPack.DetailLevels.Num() == 0) 
+			{
+				// We need at least one level !
+				FFormatNamedArguments Args;
+				Args.Add(TEXT("FullPackFilename"), FText::FromString(EachPack.MountName));
+				OutFailReason = FText::Format(LOCTEXT("NoLevelsDefined", "No detail levels defined in template pack '{FullPackFilename}'."), Args);
+				return false;
+			}
+
+			for (int32 iDetail = 0; iDetail < EachPack.DetailLevels.Num() ; iDetail++)
+			{
+				if (EachPack.DetailLevels[iDetail] == EachRequiredDetail)
+				{
+					// Build the Packname from the mount and detail
+					FString DetailString;
+					UEnum::GetValueAsString(TEXT("/Script/GameProjectGeneration.ETemplateDetailLevel"), EachRequiredDetail, DetailString);
+					
+					FString FullPackFilename = FPaths::FeaturePackDir() + EachPack.MountName +  DetailString + DefaultFeaturePackExtension;
+					if (FPaths::FileExists(FullPackFilename) == false)
+					{
+						FFormatNamedArguments Args;
+						Args.Add(TEXT("FullPackFilename"), FText::FromString(FullPackFilename));
+						OutFailReason = FText::Format(LOCTEXT("CantFindPack", "Cannot find template pack '{FullPackFilename}'."), Args);
+						return false;
+					}
+					TUniquePtr<FFeaturePackContentSource> NewContentSource = MakeUnique<FFeaturePackContentSource>(FullPackFilename, true);
+					if (NewContentSource->IsDataValid() == true)
+					{						
+						FString DestinationFolder = DestFolder;
+						FPaths::NormalizeDirectoryName(DestinationFolder);
+						bool bHasSourceFiles = false;
+						TArray<FString> FilesCopied;
+						NewContentSource->CopyAdditionalFilesToFolder(DestinationFolder, FilesCopied, bHasSourceFiles);						
+					}
+					else
+					{
+						FFormatNamedArguments Args;
+						Args.Add(TEXT("PackName"), FText::FromString(FullPackFilename));
+						OutFailReason = FText::Format(LOCTEXT("PackParseError", "Error parsing template pack '{PackName}'."), Args);
+						return false;
+					}
+				}
+			}
+		}
+	}
+	return true;
+}
+
 #undef LOCTEXT_NAMESPACE
diff --git a/Engine/Source/Editor/GameProjectGeneration/Public/GameProjectUtils.h b/Engine/Source/Editor/GameProjectGeneration/Public/GameProjectUtils.h
--- a/Engine/Source/Editor/GameProjectGeneration/Public/GameProjectUtils.h
+++ b/Engine/Source/Editor/GameProjectGeneration/Public/GameProjectUtils.h
@@ -253,6 +253,16 @@ private:
 	/** Insert any required feature packs into the DefaultGame.ini file */
 	static bool InsertFeaturePacksIntoINIFile(const FProjectInformation& InProjectInfo, FText& OutFailReason);
 
+	/* 
+	 * Insert the addition files from any feature packs specified in the temapalte defs file
+	 * @param InProjectInfo		Project infor to add content for
+	 * @param CreatedFiles		List of files we copied
+	 * @param OutFailReason		Failure reason (if any)
+	 *
+	 * @returns true if no errors
+	 */
+	static bool AddSharedContentToProject(const FProjectInformation &InProjectInfo, TArray<FString> &CreatedFiles, FText& OutFailReason);
+
 	/** Returns list of starter content files */
 	static void GetStarterContentFiles(TArray<FString>& OutFilenames);
 
@@ -460,4 +470,5 @@ private:
 
 	static TWeakPtr<SNotificationItem> UpdateGameProjectNotification;
 	static TWeakPtr<SNotificationItem> WarningProjectNameNotification;
+	static FString DefaultFeaturePackExtension;	
 };
diff --git a/Engine/Source/Editor/UnrealEd/Private/UnrealEdMisc.cpp b/Engine/Source/Editor/UnrealEd/Private/UnrealEdMisc.cpp
--- a/Engine/Source/Editor/UnrealEd/Private/UnrealEdMisc.cpp
+++ b/Engine/Source/Editor/UnrealEd/Private/UnrealEdMisc.cpp
@@ -42,6 +42,8 @@
 #include "Lightmass/LightmappedSurfaceCollection.h"
 #include "IProjectManager.h"
 #include "FeaturePackContentSource.h"
+#include "TemplateProjectDefs.h"
+#include "GameProjectUtils.h"
 
 #define LOCTEXT_NAMESPACE "UnrealEd"
 
@@ -254,6 +256,9 @@ void FUnrealEdMisc::OnInit()
 	// Check for automated build/submit option
 	const bool bDoAutomatedMapBuild = FParse::Param( ParsedCmdLine, TEXT("AutomatedMapBuild") );
 
+	// This adds mount points for the template shared resources if we are editing a template project.
+	MountTemplateSharedPaths();
+
 	// Load startup map (conditionally)
 	SlowTask.EnterProgressFrame(60);
 	{
@@ -1624,4 +1629,59 @@ void FUnrealEdMisc::OnUserDefinedChordChanged(const FUICommandInfo& CommandInfo)
 	}
 }
 
+void FUnrealEdMisc::MountTemplateSharedPaths()
+{
+	FString TemplateFilename = FPaths::GetPath(FPaths::GetProjectFilePath());
+	UTemplateProjectDefs* TemplateInfo = GameProjectUtils::LoadTemplateDefs(TemplateFilename);
+	
+	if( TemplateInfo != nullptr )
+	{
+		// Which detail level do we want to mount for editing
+		EFeaturePackDetailLevel EditDetail = EFeaturePackDetailLevel::Standard;
+		
+		// Extract the mount names and insert mount points for each of the shared packs
+		TArray<FString> AddedMountSources;
+		for (int32 iShared = 0; iShared < TemplateInfo->SharedContentPacks.Num() ; iShared++)
+		{
+			EFeaturePackDetailLevel EachEditDetail = EditDetail;
+			FFeaturePackLevelSet EachPack = TemplateInfo->SharedContentPacks[iShared];
+			if (EachPack.DetailLevels.Num() == 1)
+			{
+				// If theres only only detail level override the requirement with that
+				EachEditDetail = EachPack.DetailLevels[0];
+				UE_LOG(LogUnrealEdMisc, Warning, TEXT("Only 1 detail level defined for %s in %s."), *EachPack.MountName, *TemplateFilename );
+			}
+			else if (EachPack.DetailLevels.Num() == 0)
+			{
+				// If no levels are supplied, assume standard
+				UE_LOG(LogUnrealEdMisc, Warning, TEXT("No detail levels defined for %s in %s."), *EachPack.MountName, *TemplateFilename );
+				continue;
+			}
+			for (int32 iDetail = 0; iDetail < EachPack.DetailLevels.Num(); iDetail++)
+			{
+				if (EachPack.DetailLevels[iDetail] == EachEditDetail)
+				{
+					FString DetailString;
+					UEnum::GetValueAsString(TEXT("/Script/FeaturePackContentSource.EPackDetailLevel"), EachEditDetail, DetailString);
+					FString ShareMountName = EachPack.MountName;
+					if (AddedMountSources.Find(ShareMountName) == INDEX_NONE)
+					{
+						FString RelativePath = FString::Printf(TEXT("../../../Templates/TemplateResources/%s/%s/Content"), *DetailString, *ShareMountName);
+						if( FPaths::DirectoryExists(RelativePath))
+						{
+							FString MountName = FString::Printf(TEXT("/Game/%s/"), *ShareMountName);						
+							FPackageName::RegisterMountPoint(*MountName, *RelativePath);
+							AddedMountSources.Add(ShareMountName);
+						}
+						else
+						{
+							UE_LOG(LogUnrealEdMisc, Warning, TEXT("Cannot find path %s to mount for %s resource in %s."), *RelativePath, *EachPack.MountName, *TemplateFilename );
+						}
+					}
+				}
+			}
+		}
+	}
+}
+
 #undef  LOCTEXT_NAMESPACE
diff --git a/Engine/Source/Editor/UnrealEd/Public/UnrealEdMisc.h b/Engine/Source/Editor/UnrealEd/Public/UnrealEdMisc.h
--- a/Engine/Source/Editor/UnrealEd/Public/UnrealEdMisc.h
+++ b/Engine/Source/Editor/UnrealEd/Public/UnrealEdMisc.h
@@ -43,7 +43,10 @@ public:
 
 	/** Initalizes various systems */
 	virtual void OnInit();
-	
+
+	/* Check if this we are editing a template project, and if so mount any shared resource paths it uses */
+	void MountTemplateSharedPaths();
+
 	/* Cleans up various systems */
 	virtual void OnExit();
 
diff --git a/Engine/Source/Editor/UnrealEd/UnrealEd.Build.cs b/Engine/Source/Editor/UnrealEd/UnrealEd.Build.cs
--- a/Engine/Source/Editor/UnrealEd/UnrealEd.Build.cs
+++ b/Engine/Source/Editor/UnrealEd/UnrealEd.Build.cs
@@ -132,6 +132,7 @@ public class UnrealEd : ModuleRules
                 "MaterialUtilities",
                 "LocalizationService",
 				"AddContentDialog",
+				"GameProjectGeneration",
 			}
 		);
 
