Minor revision to shared feature pack asset code.

diff --git a/Engine/Source/Editor/AddContentDialog/Private/ContentSourceProviders/FeaturePack/FeaturePackContentSource.cpp b/Engine/Source/Editor/AddContentDialog/Private/ContentSourceProviders/FeaturePack/FeaturePackContentSource.cpp
--- a/Engine/Source/Editor/AddContentDialog/Private/ContentSourceProviders/FeaturePack/FeaturePackContentSource.cpp
+++ b/Engine/Source/Editor/AddContentDialog/Private/ContentSourceProviders/FeaturePack/FeaturePackContentSource.cpp
@@ -344,6 +344,9 @@ bool FFeaturePackContentSource::InstallToProject(FString InstallPath)
 	}
 	else
 	{
+		// We need to insert additional packs before we import the main assets since the code in the main pack may reference them
+		InsertAdditionalFeaturePacks();
+
 		FAssetToolsModule& AssetToolsModule = FModuleManager::Get().LoadModuleChecked<FAssetToolsModule>("AssetTools");
 		TArray<FString> AssetPaths;
 		AssetPaths.Add(FeaturePackPath);
@@ -362,27 +365,7 @@ bool FFeaturePackContentSource::InstallToProject(FString InstallPath)
 				ToSave.AddUnique(ImportedObject->GetOutermost());
 			}
 			FEditorFileUtils::PromptForCheckoutAndSave( ToSave, /*bCheckDirty=*/ false, /*bPromptToSave=*/ false );
-
-			// Now copy files in any additional packs (These are for the shared assets)
- 			EFeaturePackDetailLevel RequiredLevel = EFeaturePackDetailLevel::High;
- 			for (int32 iExtraPack = 0; iExtraPack < AdditionalFeaturePacks.Num(); iExtraPack++)
- 			{
-				FString DestinationFolder = FPaths::GameDir();
-				FString FullPath = FPaths::FeaturePackDir() + AdditionalFeaturePacks[iExtraPack].GetFeaturePackNameForLevel(RequiredLevel);
- 				if (FPlatformFileManager::Get().GetPlatformFile().FileExists(*FullPath))
- 				{
- 					FString InsertPath = DestinationFolder + *AdditionalFeaturePacks[iExtraPack].MountName;
-
-					TUniquePtr<FFeaturePackContentSource> NewContentSource = MakeUnique<FFeaturePackContentSource>(FullPath, true);
-					if (NewContentSource->IsDataValid() == true)
-					{
-						bool bHasSourceFiles = false;
-						TArray<FString> FilesCopied;
-						NewContentSource->CopyAdditionalFilesToFolder(DestinationFolder, FilesCopied, bHasSourceFiles);
-					}
- 				}
- 			}
-
+			
 			bResult = true;
 			
 			// Focus on a specific asset if we want to.
@@ -720,4 +703,27 @@ void FFeaturePackContentSource::CopyAdditionalFilesToFolder( const FString& Dest
 	}
 }
 
+void FFeaturePackContentSource::InsertAdditionalFeaturePacks()
+{
+	// Now copy files in any additional packs (These are for the shared assets)
+	EFeaturePackDetailLevel RequiredLevel = EFeaturePackDetailLevel::High;
+	for (int32 iExtraPack = 0; iExtraPack < AdditionalFeaturePacks.Num(); iExtraPack++)
+	{
+		FString DestinationFolder = FPaths::GameDir();
+		FString FullPath = FPaths::FeaturePackDir() + AdditionalFeaturePacks[iExtraPack].GetFeaturePackNameForLevel(RequiredLevel);
+		if (FPlatformFileManager::Get().GetPlatformFile().FileExists(*FullPath))
+		{
+			FString InsertPath = DestinationFolder + *AdditionalFeaturePacks[iExtraPack].MountName;
+
+			TUniquePtr<FFeaturePackContentSource> NewContentSource = MakeUnique<FFeaturePackContentSource>(FullPath, true);
+			if (NewContentSource->IsDataValid() == true)
+			{
+				bool bHasSourceFiles = false;
+				TArray<FString> FilesCopied;
+				NewContentSource->CopyAdditionalFilesToFolder(DestinationFolder, FilesCopied, bHasSourceFiles);
+			}
+		}
+	}
+}
+
 #undef LOCTEXT_NAMESPACE 
diff --git a/Engine/Source/Editor/AddContentDialog/Public/FeaturePackContentSource.h b/Engine/Source/Editor/AddContentDialog/Public/FeaturePackContentSource.h
--- a/Engine/Source/Editor/AddContentDialog/Public/FeaturePackContentSource.h
+++ b/Engine/Source/Editor/AddContentDialog/Public/FeaturePackContentSource.h
@@ -127,6 +127,9 @@ public:
 	virtual TArray<TSharedPtr<FImageData>> GetScreenshotData() const override;
 
 	virtual bool InstallToProject(FString InstallPath) override;
+
+	void InsertAdditionalFeaturePacks();
+
 	virtual bool IsDataValid() const override;
 	
 	virtual ~FFeaturePackContentSource();
