fixed shadow var

diff --git a/Engine/Source/Editor/AddContentDialog/Private/ContentSourceProviders/FeaturePack/FeaturePackContentSource.cpp b/Engine/Source/Editor/AddContentDialog/Private/ContentSourceProviders/FeaturePack/FeaturePackContentSource.cpp
--- a/Engine/Source/Editor/AddContentDialog/Private/ContentSourceProviders/FeaturePack/FeaturePackContentSource.cpp
+++ b/Engine/Source/Editor/AddContentDialog/Private/ContentSourceProviders/FeaturePack/FeaturePackContentSource.cpp
@@ -258,8 +258,8 @@ FFeaturePackContentSource::FFeaturePackContentSource(FString InFeaturePackPath,
 			for (TSharedPtr<FJsonValue> DetailValue: EachAdditionalPack->GetArrayField("DetailLevels"))
 			{
 				const FString DetailString = DetailValue->AsString();
-				int32 Index = DetailEnum->FindEnumIndex(FName(*DetailString));
-				EFeaturePackDetailLevel EachLevel = Index != INDEX_NONE ? (EFeaturePackDetailLevel)Index : EFeaturePackDetailLevel::Standard;
+				int32 eIndex = DetailEnum->FindEnumIndex(FName(*DetailString));
+				EFeaturePackDetailLevel EachLevel = Index != INDEX_NONE ? (EFeaturePackDetailLevel)eIndex : EFeaturePackDetailLevel::Standard;
 				Levels.AddUnique(EachLevel);
 			}
 			int32 PackIndex = AdditionalFeaturePacks.Add(FFeaturePackLevelSet(
