AddContentDialog - Fix LogSlateStyle errors

Jira: UE-16603

diff --git a/Engine/Source/Editor/AddContentDialog/Private/SAddContentWidget.cpp b/Engine/Source/Editor/AddContentDialog/Private/SAddContentWidget.cpp
--- a/Engine/Source/Editor/AddContentDialog/Private/SAddContentWidget.cpp
+++ b/Engine/Source/Editor/AddContentDialog/Private/SAddContentWidget.cpp
@@ -282,8 +282,8 @@ TSharedRef<SWidget> SAddContentWidget::CreateContentSourceDetail(TSharedPtr<FCon
 TSharedRef<SWidget> SAddContentWidget::CreateScreenshotCarousel(TSharedPtr<FContentSourceViewModel> ContentSource)
 {
 	return SNew(SWidgetCarouselWithNavigation<TSharedPtr<FSlateBrush>>)
-		.NavigationBarStyle(FWidgetCarouselModuleStyle::Get(), "CarouselNavigationButton")
-		.NavigationButtonStyle(FWidgetCarouselModuleStyle::Get(), "CarouselNavigationBar")
+		.NavigationBarStyle(FWidgetCarouselModuleStyle::Get(), "CarouselNavigationBar")
+		.NavigationButtonStyle(FWidgetCarouselModuleStyle::Get(), "CarouselNavigationButton")
 		.OnGenerateWidget(this, &SAddContentWidget::CreateScreenshotWidget)
 		.WidgetItemsSource(ContentSource->GetScreenshotBrushes());
 }
