Revised feature packs to handle supersearch registration correctly.

diff --git a/Engine/Source/Editor/AddContentDialog/Private/ContentSourceProviders/FeaturePack/FeaturePackContentSource.cpp b/Engine/Source/Editor/AddContentDialog/Private/ContentSourceProviders/FeaturePack/FeaturePackContentSource.cpp
--- a/Engine/Source/Editor/AddContentDialog/Private/ContentSourceProviders/FeaturePack/FeaturePackContentSource.cpp
+++ b/Engine/Source/Editor/AddContentDialog/Private/ContentSourceProviders/FeaturePack/FeaturePackContentSource.cpp
@@ -281,8 +281,11 @@ FFeaturePackContentSource::FFeaturePackContentSource(FString InFeaturePackPath,
 	if( bDontRegisterForSearch == false )
 	{
 		FSuperSearchModule& SuperSearchModule = FModuleManager::LoadModuleChecked< FSuperSearchModule >(TEXT("SuperSearch"));
-		SuperSearchModule.GetActOnSearchTextClicked().AddRaw(this, &FFeaturePackContentSource::HandleActOnSearchText);
-		SuperSearchModule.GetSearchTextChanged().AddRaw(this, &FFeaturePackContentSource::HandleSuperSearchTextChanged);
+		// Remove any existing delegates for this pack
+		FCoreUObjectDelegates::OnObjectPropertyChanged.Remove(SearchClickedHandle);
+		FCoreUObjectDelegates::OnObjectPropertyChanged.Remove(SearchChangedHandle);
+		SearchClickedHandle = SuperSearchModule.GetActOnSearchTextClicked().AddRaw(this, &FFeaturePackContentSource::HandleActOnSearchText);
+		SearchChangedHandle = SuperSearchModule.GetSearchTextChanged().AddRaw(this, &FFeaturePackContentSource::HandleSuperSearchTextChanged);
 	}
 	bPackValid = true;
 }
@@ -387,6 +390,9 @@ bool FFeaturePackContentSource::InstallToProject(FString InstallPath)
 
 FFeaturePackContentSource::~FFeaturePackContentSource()
 {
+	// Remove any search handler delegates for this pack
+	FCoreUObjectDelegates::OnObjectPropertyChanged.Remove(SearchClickedHandle);
+	FCoreUObjectDelegates::OnObjectPropertyChanged.Remove(SearchChangedHandle);
 }
 
 bool FFeaturePackContentSource::IsDataValid() const
diff --git a/Engine/Source/Editor/AddContentDialog/Public/FeaturePackContentSource.h b/Engine/Source/Editor/AddContentDialog/Public/FeaturePackContentSource.h
--- a/Engine/Source/Editor/AddContentDialog/Public/FeaturePackContentSource.h
+++ b/Engine/Source/Editor/AddContentDialog/Public/FeaturePackContentSource.h
@@ -114,6 +114,7 @@ class FFeaturePackContentSource : public IContentSource
 {
 public:
 	ADDCONTENTDIALOG_API FFeaturePackContentSource(FString InFeaturePackPath, bool bDontRegister = false);
+	virtual ~FFeaturePackContentSource();
 
 	virtual TArray<FLocalizedText> GetLocalizedNames() const override;
 	virtual TArray<FLocalizedText> GetLocalizedDescriptions() const override;
@@ -132,7 +133,6 @@ public:
 
 	virtual bool IsDataValid() const override;
 	
-	virtual ~FFeaturePackContentSource();
 	
 	void HandleActOnSearchText(TSharedPtr<FSearchEntry> SearchEntry);
 	void HandleSuperSearchTextChanged(const FString& InText, TArray< TSharedPtr<FSearchEntry> >& OutSuggestions);
@@ -193,4 +193,6 @@ private:
 	FString SortKey;
 	TArray<FLocalizedTextArray> LocalizedSearchTags;
 	TArray<FFeaturePackLevelSet> AdditionalFeaturePacks;
+	FDelegateHandle SearchClickedHandle;
+	FDelegateHandle SearchChangedHandle;
 };
\ No newline at end of file
