Added code to filter unwanted feature packs from add content dialog
Corrected preview images in shared packs (we dont use these but at least they are correct if we decide we want to)
Removed some defunct preview images.
Removed superfluous , in template manifests
Repacked upacks

diff --git a/Engine/Source/Editor/AddContentDialog/Private/ViewModels/AddContentWidgetViewModel.cpp b/Engine/Source/Editor/AddContentDialog/Private/ViewModels/AddContentWidgetViewModel.cpp
--- a/Engine/Source/Editor/AddContentDialog/Private/ViewModels/AddContentWidgetViewModel.cpp
+++ b/Engine/Source/Editor/AddContentDialog/Private/ViewModels/AddContentWidgetViewModel.cpp
@@ -109,16 +109,25 @@ void FAddContentWidgetViewModel::BuildContentSourceViewModels()
 	FilteredContentSourceViewModels.Empty();
 	CategoryToSelectedContentSourceMap.Empty();
 	
+	// List of categorys we dont want to see
+	TArray<EContentSourceCategory> FilteredCategories;
+	FilteredCategories.Add(EContentSourceCategory::SharedPack);
+	FilteredCategories.Add(EContentSourceCategory::Unknown);
+
 	for (auto ContentSourceProvider : ContentSourceProviders)
 	{
 		for (auto ContentSource : ContentSourceProvider->GetContentSources())
 		{
-			TSharedPtr<FContentSourceViewModel> ContentSourceViewModel = MakeShareable(new FContentSourceViewModel(ContentSource));
-			if (Categories.Contains(ContentSourceViewModel->GetCategory()) == false)
+			// Check if we want to see this category
+			if( FilteredCategories.Contains( ContentSource->GetCategory()) == false )
 			{
-				Categories.Add(ContentSourceViewModel->GetCategory());
+				TSharedPtr<FContentSourceViewModel> ContentSourceViewModel = MakeShareable(new FContentSourceViewModel(ContentSource));
+				if (Categories.Contains(ContentSourceViewModel->GetCategory()) == false)
+				{
+					Categories.Add(ContentSourceViewModel->GetCategory());
+				}
+				ContentSourceViewModels.Add(ContentSourceViewModel);
 			}
-			ContentSourceViewModels.Add(ContentSourceViewModel);
 		}
 	}
 
