Fixed crash in FeaturePackContenSource.cpp

diff --git a/Engine/Source/Editor/AddContentDialog/Private/ContentSourceProviders/FeaturePack/FeaturePackContentSource.cpp b/Engine/Source/Editor/AddContentDialog/Private/ContentSourceProviders/FeaturePack/FeaturePackContentSource.cpp
--- a/Engine/Source/Editor/AddContentDialog/Private/ContentSourceProviders/FeaturePack/FeaturePackContentSource.cpp
+++ b/Engine/Source/Editor/AddContentDialog/Private/ContentSourceProviders/FeaturePack/FeaturePackContentSource.cpp
@@ -272,7 +272,7 @@ FFeaturePackContentSource::FFeaturePackContentSource(FString InFeaturePackPath,
 				FString FullPath = FPaths::FeaturePackDir() + NewSet.GetFeaturePackNameForLevel(NewSet.DetailLevels[iLevel]);	
 				if (!FPlatformFileManager::Get().GetPlatformFile().FileExists(*FullPath))
 				{
-					RecordAndLogError( FString::Printf(TEXT("Error in Feature pack %s. Cannot find additional pack %s."), *FullPath));
+					RecordAndLogError( FString::Printf(TEXT("Error in Feature pack %s. Cannot find additional pack %s."), *InFeaturePackPath, *FullPath));
 				}
 			}		
 		}
