Revised feature pack code.
Additional files can now be specified in feature pack manifest.
Feature packs can now be parsed/installed feature pack without the need for a .upack file
Removed feature packs for shared template resources
Moved shared template resource manifests into a FeaturePack folder and revised so they will work sans .upack
Removed shared feature packs from list of .upacks rocket builds.

diff --git a/Engine/Source/Editor/AddContentDialog/Private/ContentSourceProviders/FeaturePack/FeaturePackContentSource.cpp b/Engine/Source/Editor/AddContentDialog/Private/ContentSourceProviders/FeaturePack/FeaturePackContentSource.cpp
--- a/Engine/Source/Editor/AddContentDialog/Private/ContentSourceProviders/FeaturePack/FeaturePackContentSource.cpp
+++ b/Engine/Source/Editor/AddContentDialog/Private/ContentSourceProviders/FeaturePack/FeaturePackContentSource.cpp
@@ -9,6 +9,7 @@
 #include "FileHelpers.h"
 #include "Editor/MainFrame/Public/Interfaces/IMainFrameModule.h"
 #include "SuperSearchModule.h"
+#include "FeaturePackContentSourceProvider.h"
 
 #define LOCTEXT_NAMESPACE "ContentFeaturePacks"
 
@@ -118,6 +119,21 @@ bool TryValidateManifestObject(TSharedPtr<FJsonObject> ManifestObject, TSharedPt
 		return false;
 	}
 
+	// If we have an additional files entry check its valid
+	if (ManifestObject->HasTypedField<EJson::Object>("AdditionalFiles") == true)
+	{
+		TSharedPtr<FJsonObject> AdditionalFileObject = ManifestObject->GetObjectField("AdditionalFiles");
+		if (AdditionalFileObject->HasTypedField<EJson::String>("DestinationFilesFolder") == false )
+		{
+			ErrorMessage = MakeShareable(new FString("Manifest has an AdditionalFiles object but no DestinationFilesFolder"));
+			return false;
+		}
+		if (AdditionalFileObject->HasTypedField<EJson::Array>("AdditionalFilesList") == false)
+		{
+			ErrorMessage = MakeShareable(new FString("Manifest has an AdditionalFiles object but no AdditionalFilesList"));
+			return false;
+		}
+	}
 	return true;
 }
 
@@ -125,169 +141,68 @@ FFeaturePackContentSource::FFeaturePackContentSource(FString InFeaturePackPath,
 {
 	FeaturePackPath = InFeaturePackPath;
 	bPackValid = false;
-	// Create a pak platform file and mount the feature pack file.
-	FPakPlatformFile PakPlatformFile;
-	PakPlatformFile.Initialize(&FPlatformFileManager::Get().GetPlatformFile(), TEXT(""));
-	FString MountPoint = "root:/";
-	PakPlatformFile.Mount(*InFeaturePackPath, 0, *MountPoint);
 	
-	// Gets the manifest file as a JSon string
-	TArray<uint8> ManifestBuffer;
-	if( LoadPakFileToBuffer(PakPlatformFile, FPaths::Combine(*MountPoint, TEXT("manifest.json")), ManifestBuffer) == false )
-	{
-		RecordAndLogError( FString::Printf(TEXT("Error in Feature pack %s. Cannot find manifest."), *InFeaturePackPath));
-		Category = EContentSourceCategory::Unknown;
-		return;
-	}
 	FString ManifestString;
-	FFileHelper::BufferToString(ManifestString, ManifestBuffer.GetData(), ManifestBuffer.Num());
-
-	// Populate text fields from the manifest.
-	TSharedPtr<FJsonObject> ManifestObject;
-	TSharedRef<TJsonReader<>> ManifestReader = TJsonReaderFactory<>::Create(ManifestString);
-	FJsonSerializer::Deserialize(ManifestReader, ManifestObject);
-
-	if (ManifestReader->GetErrorMessage().IsEmpty() == false)
-	{
-		RecordAndLogError( FString::Printf(TEXT("Error in Feature pack %s. Failed to parse manifest: %s"), *InFeaturePackPath, *ManifestReader->GetErrorMessage()));
-		Category = EContentSourceCategory::Unknown;
-		return;
-	}
-
-	TSharedPtr<FString> ManifestObjectErrorMessage;
-	if (TryValidateManifestObject(ManifestObject, ManifestObjectErrorMessage) == false)
-	{
-		RecordAndLogError( FString::Printf(TEXT("Error in Feature pack %s. Manifest object error: %s"), *InFeaturePackPath, **ManifestObjectErrorMessage));
-		Category = EContentSourceCategory::Unknown;
-		return;
-	}
-
-	for (TSharedPtr<FJsonValue> NameValue : ManifestObject->GetArrayField("Name"))
-	{
-		TSharedPtr<FJsonObject> LocalizedNameObject = NameValue->AsObject();
-		LocalizedNames.Add(FLocalizedText(
-			LocalizedNameObject->GetStringField("Language"),
-			FText::FromString(LocalizedNameObject->GetStringField("Text"))));
-	}
-
-	for (TSharedPtr<FJsonValue> DescriptionValue : ManifestObject->GetArrayField("Description"))
+	Category = EContentSourceCategory::Unknown;
+	if( InFeaturePackPath.EndsWith(TEXT(".upack") ) == true )
 	{
-		TSharedPtr<FJsonObject> LocalizedDescriptionObject = DescriptionValue->AsObject();
-		LocalizedDescriptions.Add(FLocalizedText(
-			LocalizedDescriptionObject->GetStringField("Language"),
-			FText::FromString(LocalizedDescriptionObject->GetStringField("Text"))));
-	}
+		bContentsInPakFile = true;
+		MountPoint = "root:/";
+		// Create a pak platform file and mount the feature pack file.
+		FPakPlatformFile PakPlatformFile;
+		PakPlatformFile.Initialize(&FPlatformFileManager::Get().GetPlatformFile(), TEXT(""));
+		
+		PakPlatformFile.Mount(*InFeaturePackPath, 0, *MountPoint);
 
-	// Parse asset types field
-	for (TSharedPtr<FJsonValue> AssetTypesValue : ManifestObject->GetArrayField("AssetTypes"))
-	{
-		TSharedPtr<FJsonObject> LocalizedAssetTypesObject = AssetTypesValue->AsObject();
-		LocalizedAssetTypesList.Add(FLocalizedText(
-			LocalizedAssetTypesObject->GetStringField("Language"),
-			FText::FromString(LocalizedAssetTypesObject->GetStringField("Text"))));
-	}
-	
-	// Parse search tags field
-	if( ManifestObject->HasField("SearchTags")==true)
-	{
-		for (TSharedPtr<FJsonValue> AssetTypesValue : ManifestObject->GetArrayField("SearchTags"))
+		// Gets the manifest file as a JSon string
+		TArray<uint8> ManifestBuffer;
+		if (LoadPakFileToBuffer(PakPlatformFile, FPaths::Combine(*MountPoint, TEXT("manifest.json")), ManifestBuffer) == false)
 		{
-			TSharedPtr<FJsonObject> LocalizedAssetTypesObject = AssetTypesValue->AsObject();
-			LocalizedSearchTags.Add(FLocalizedTextArray(
-				LocalizedAssetTypesObject->GetStringField("Language"),
-				LocalizedAssetTypesObject->GetStringField("Text")));
+			RecordAndLogError(FString::Printf(TEXT("Error in Feature pack %s. Cannot find manifest."), *FeaturePackPath));
+			return;
 		}
-	}
-
-	// Parse class types field
-	ClassTypes = ManifestObject->GetStringField("ClassTypes");	
-	
-	// Parse initial focus asset if we have one - this is not required
-	if (ManifestObject->HasTypedField<EJson::String>("FocusAsset") == true)
-	{
-		FocusAssetIdent = ManifestObject->GetStringField("FocusAsset");
-	}	
-
-	// Use the path as the sort key - it will be alphabetical that way
-	SortKey = FeaturePackPath;
-	ManifestObject->TryGetStringField("SortKey", SortKey);
-
-	FString CategoryString = ManifestObject->GetStringField("Category");	
-	UEnum* Enum = FindObjectChecked<UEnum>(ANY_PACKAGE, TEXT("EContentSourceCategory"));
-	int32 Index = Enum->FindEnumIndex(FName(*CategoryString));
-	Category = Index != INDEX_NONE ? (EContentSourceCategory)Index : EContentSourceCategory::Unknown;
+		
+		FFileHelper::BufferToString(ManifestString, ManifestBuffer.GetData(), ManifestBuffer.Num());
 
-	// Load image data
-	FString IconFilename = ManifestObject->GetStringField("Thumbnail");
-	TSharedPtr<TArray<uint8>> IconImageData = MakeShareable(new TArray<uint8>());
-	FString ThumbnailFile = FPaths::Combine(*MountPoint, TEXT("Media"), *IconFilename);
-	if( LoadPakFileToBuffer(PakPlatformFile, ThumbnailFile, *IconImageData) == true)
-	{
-		IconData = MakeShareable(new FImageData(IconFilename, IconImageData));
+		if (ParseManifestString(ManifestString) == true)
+		{
+			LoadFeaturePackImageDataFromPackFile(PakPlatformFile);			
+		}
 	}
 	else
 	{
-		RecordAndLogError( FString::Printf(TEXT("Error in Feature pack %s. Cannot find thumbnail %s."), *InFeaturePackPath, *ThumbnailFile ));
-	}
-
-	// parse the screenshots field
-	const TArray<TSharedPtr<FJsonValue>> ScreenshotFilenameArray = ManifestObject->GetArrayField("Screenshots");
-	for (const TSharedPtr<FJsonValue> ScreenshotFilename : ScreenshotFilenameArray)
-	{
-		TSharedPtr<TArray<uint8>> SingleScreenshotData = MakeShareable(new TArray<uint8>);
-		if( LoadPakFileToBuffer(PakPlatformFile, FPaths::Combine(*MountPoint, TEXT("Media"), *ScreenshotFilename->AsString()), *SingleScreenshotData) )
+		bContentsInPakFile = false;
+		FString TemplatesFolder = TEXT("FeaturePack");
+		FString ThisTemplateRoot = FPaths::GetPath(FeaturePackPath);
+		if( ThisTemplateRoot.EndsWith(TemplatesFolder) == true )
 		{
-			ScreenshotData.Add(MakeShareable(new FImageData(ScreenshotFilename->AsString(), SingleScreenshotData)));
-		}
-		else
-		{
-			RecordAndLogError( FString::Printf(TEXT("Error in Feature pack %s. Cannot find screenshot %s."), *InFeaturePackPath, *ScreenshotFilename->AsString() ));
+			int32 Index = ThisTemplateRoot.Find(TemplatesFolder);
+			ThisTemplateRoot = ThisTemplateRoot.Left(Index);
 		}
-	}
-
-	// Parse additional packs data
-	if (ManifestObject->HasTypedField<EJson::Array>("AdditionalFeaturePacks") == true)
-	{		
-		UEnum* DetailEnum = FindObjectChecked<UEnum>(ANY_PACKAGE, TEXT("EFeaturePackDetailLevel"));
-		for (TSharedPtr<FJsonValue> AdditionalFeaturePackValue : ManifestObject->GetArrayField("AdditionalFeaturePacks"))
+		MountPoint = ThisTemplateRoot;
+		FFileHelper::LoadFileToString( ManifestString, *FeaturePackPath);
+		if (ParseManifestString(ManifestString) == true)
 		{
-			TSharedPtr<FJsonObject> EachAdditionalPack = AdditionalFeaturePackValue->AsObject();
-			FString MountName = EachAdditionalPack->GetStringField("MountName");
-
-			TArray<EFeaturePackDetailLevel>	Levels;
-			for (TSharedPtr<FJsonValue> DetailValue: EachAdditionalPack->GetArrayField("DetailLevels"))
-			{
-				const FString DetailString = DetailValue->AsString();
-				int32 eIndex = DetailEnum->FindEnumIndex(FName(*DetailString));
-				EFeaturePackDetailLevel EachLevel = Index != INDEX_NONE ? (EFeaturePackDetailLevel)eIndex : EFeaturePackDetailLevel::Standard;
-				Levels.AddUnique(EachLevel);
-			}
-			int32 PackIndex = AdditionalFeaturePacks.Add(FFeaturePackLevelSet(
-				MountName,
-				Levels));
-			// Check that the additional feature packs listed exist
-			FFeaturePackLevelSet& NewSet = AdditionalFeaturePacks[PackIndex];
-			for (int32 iLevel = 0; iLevel < NewSet.DetailLevels.Num() ; iLevel++)
-			{
-				FString FullPath = FPaths::FeaturePackDir() + NewSet.GetFeaturePackNameForLevel(NewSet.DetailLevels[iLevel]);	
-				if (!FPlatformFileManager::Get().GetPlatformFile().FileExists(*FullPath))
-				{
-					RecordAndLogError( FString::Printf(TEXT("Error in Feature pack %s. Cannot find additional pack %s."), *InFeaturePackPath, *FullPath));
-				}
-			}		
+			LoadFeaturePackImageData();			
 		}
 	}
 
-	if( bDontRegisterForSearch == false )
+	if( (bDontRegisterForSearch == false) && ( bPackValid == true ) )
 	{
-		FSuperSearchModule& SuperSearchModule = FModuleManager::LoadModuleChecked< FSuperSearchModule >(TEXT("SuperSearch"));
+		FSuperSearchModule& SuperSearchModule = FModuleManager::LoadModuleChecked< FSuperSearchModule >("SuperSearch");
 		// Remove any existing delegates for this pack
-		FCoreUObjectDelegates::OnObjectPropertyChanged.Remove(SearchClickedHandle);
-		FCoreUObjectDelegates::OnObjectPropertyChanged.Remove(SearchChangedHandle);
+		SuperSearchModule.GetActOnSearchTextClicked().Remove(SearchClickedHandle);
+		SuperSearchModule.GetSearchTextChanged().Remove(SearchChangedHandle);
 		SearchClickedHandle = SuperSearchModule.GetActOnSearchTextClicked().AddRaw(this, &FFeaturePackContentSource::HandleActOnSearchText);
 		SearchChangedHandle = SuperSearchModule.GetSearchTextChanged().AddRaw(this, &FFeaturePackContentSource::HandleSuperSearchTextChanged);
 	}
-	bPackValid = true;
+}
+
+FFeaturePackContentSource::FFeaturePackContentSource()
+{
+	bPackValid = false;
+	TSharedPtr<FJsonObject> ManifestObject;
+	bPackValid = false;
 }
 
 bool FFeaturePackContentSource::LoadPakFileToBuffer(FPakPlatformFile& PakPlatformFile, FString Path, TArray<uint8>& Buffer)
@@ -348,51 +263,69 @@ bool FFeaturePackContentSource::InstallToProject(FString InstallPath)
 	else
 	{
 		// We need to insert additional packs before we import the main assets since the code in the main pack may reference them
-		InsertAdditionalFeaturePacks();
+		// TODO - handle errors from this
+		TArray<FString> FilesCopied;
+		InsertAdditionalResources(AdditionalFeaturePacks,EFeaturePackDetailLevel::High, FPaths::GameDir(),FilesCopied);
+
+ 		if (AdditionalFilesForPack.AdditionalFilesList.Num() != 0)
+ 		{
+ 			bool bHasSourceFiles = false;
+ 			TArray<FString> FilesCopied;
+ 			CopyAdditionalFilesToFolder( FPaths::GameDir(), FilesCopied, bHasSourceFiles/*,AdditionalFilesFolder*/);
+ 		}
+
+		if( bContentsInPakFile == true)
+		{
+			FAssetToolsModule& AssetToolsModule = FModuleManager::Get().LoadModuleChecked<FAssetToolsModule>("AssetTools");
+			TArray<FString> AssetPaths;
+			AssetPaths.Add(FeaturePackPath);
+			TArray<UObject*> ImportedObjects = AssetToolsModule.Get().ImportAssets(AssetPaths, InstallPath);
 
-		FAssetToolsModule& AssetToolsModule = FModuleManager::Get().LoadModuleChecked<FAssetToolsModule>("AssetTools");
-		TArray<FString> AssetPaths;
-		AssetPaths.Add(FeaturePackPath);
 
-		TArray<UObject*> ImportedObjects = AssetToolsModule.Get().ImportAssets(AssetPaths, InstallPath );
-		if( ImportedObjects.Num() == 0 )
-		{
-			UE_LOG(LogFeaturePack, Warning, TEXT("No objects imported installing pack %s"), *InstallPath);
-		}
-		else
-		{
-			// Save any imported assets.
-			TArray<UPackage*> ToSave;
-			for (auto ImportedObject : ImportedObjects)
+			if (ImportedObjects.Num() == 0)
 			{
-				ToSave.AddUnique(ImportedObject->GetOutermost());
+				UE_LOG(LogFeaturePack, Warning, TEXT("No objects imported installing pack %s"), *InstallPath);
 			}
-			FEditorFileUtils::PromptForCheckoutAndSave( ToSave, /*bCheckDirty=*/ false, /*bPromptToSave=*/ false );
-			
-			bResult = true;
-			
-			// Focus on a specific asset if we want to.
-			if( GetFocusAssetName().IsEmpty() == false )
+			else
 			{
-				UObject* FocusAsset = LoadObject<UObject>(nullptr, *GetFocusAssetName());
-				if (FocusAsset)
+				// Save any imported assets.
+				TArray<UPackage*> ToSave;
+				for (auto ImportedObject : ImportedObjects)
 				{
-					FContentBrowserModule& ContentBrowserModule = FModuleManager::Get().LoadModuleChecked<FContentBrowserModule>("ContentBrowser");
-					TArray<UObject*> SyncObjects;
-					SyncObjects.Add(FocusAsset);
-					ContentBrowserModule.Get().SyncBrowserToAssets(SyncObjects);
+					ToSave.AddUnique(ImportedObject->GetOutermost());
 				}
+				FEditorFileUtils::PromptForCheckoutAndSave(ToSave, /*bCheckDirty=*/ false, /*bPromptToSave=*/ false);
+
+				bResult = true;
 			}
 		}
+		
+		// Focus on a specific asset if we want to.
+		if (GetFocusAssetName().IsEmpty() == false)
+		{
+			UObject* FocusAsset = LoadObject<UObject>(nullptr, *GetFocusAssetName());
+			if (FocusAsset)
+			{
+				FContentBrowserModule& ContentBrowserModule = FModuleManager::Get().LoadModuleChecked<FContentBrowserModule>("ContentBrowser");
+				TArray<UObject*> SyncObjects;
+				SyncObjects.Add(FocusAsset);
+				ContentBrowserModule.Get().SyncBrowserToAssets(SyncObjects);
+			}
+		}
+
 	}
 	return bResult;
 }
 
 FFeaturePackContentSource::~FFeaturePackContentSource()
 {
-	// Remove any search handler delegates for this pack
-	FCoreUObjectDelegates::OnObjectPropertyChanged.Remove(SearchClickedHandle);
-	FCoreUObjectDelegates::OnObjectPropertyChanged.Remove(SearchChangedHandle);
+	// Remove any search handler delegates for this pack	
+	if( FModuleManager::Get().IsModuleLoaded("SuperSearch")==true)
+	{
+		FSuperSearchModule& SuperSearchModule = FModuleManager::Get().GetModuleChecked< FSuperSearchModule >("SuperSearch");
+		SuperSearchModule.GetActOnSearchTextClicked().Remove(SearchClickedHandle);
+		SuperSearchModule.GetSearchTextChanged().Remove(SearchChangedHandle);
+	}
 }
 
 bool FFeaturePackContentSource::IsDataValid() const
@@ -415,6 +348,10 @@ FString FFeaturePackContentSource::GetSortKey() const
 	return SortKey;
 }
 
+FString FFeaturePackContentSource::GetIdent() const
+{
+	return Identity;
+}
 
 void FFeaturePackContentSource::HandleActOnSearchText(TSharedPtr<FSearchEntry> SearchEntry)
 {
@@ -522,77 +459,97 @@ bool FFeaturePackContentSource::GetAdditionalFilesForPack(TArray<FString>& FileL
 	}
 	else
 	{
-		// Create a pak platform file and mount the feature pack file.
-		FPakPlatformFile PakPlatformFile;
-		PakPlatformFile.Initialize(&FPlatformFileManager::Get().GetPlatformFile(), TEXT(""));
-		FString MountPoint = "root:/";
-		PakPlatformFile.Mount(*FeaturePackPath, 0, *MountPoint);
-
-		// Gets the manifest file as a JSon string
-		TArray<uint8> ManifestBuffer;
-		if (LoadPakFileToBuffer(PakPlatformFile, FPaths::Combine(*MountPoint, TEXT("Config/Config.ini")), ManifestBuffer) == false)
+		// This doesnt support additional files in the manifest and in the config file. Should only do one or the other.
+		if( AdditionalFilesForPack.AdditionalFilesList.Num() != 0 )
 		{
-			RecordAndLogError(FString::Printf(TEXT("Error in Feature pack %s. Cannot find Config.ini"), *FeaturePackPath));			
+			BuildListOfAdditionalFiles(AdditionalFilesForPack.AdditionalFilesList,FileList,bContainsSource);
 		}
-		else
+		else if( bContentsInPakFile == true)
 		{
-			FString ConfigFileString;
-			FFileHelper::BufferToString(ConfigFileString, ManifestBuffer.GetData(), ManifestBuffer.Num());
-			bParsedFiles = ExtractListOfAdditionalFiles(ConfigFileString, FileList, bContainsSource);
+			// Create a pak platform file and mount the feature pack file.
+			FPakPlatformFile PakPlatformFile;
+			PakPlatformFile.Initialize(&FPlatformFileManager::Get().GetPlatformFile(), TEXT(""));
+			FString MountPoint = "root:/";
+			PakPlatformFile.Mount(*FeaturePackPath, 0, *MountPoint);
+
+			// Gets the manifest file as a JSon string
+			TArray<uint8> ManifestBuffer;
+			if (LoadPakFileToBuffer(PakPlatformFile, FPaths::Combine(*MountPoint, TEXT("Config/Config.ini")), ManifestBuffer) == false)
+			{
+				RecordAndLogError(FString::Printf(TEXT("Error in Feature pack %s. Cannot find Config.ini"), *FeaturePackPath));
+			}
+			else
+			{
+				FString ConfigFileString;
+				FFileHelper::BufferToString(ConfigFileString, ManifestBuffer.GetData(), ManifestBuffer.Num());
+				bParsedFiles = ExtractListOfAdditionalFiles(ConfigFileString, FileList, bContainsSource);
+			}
 		}
 	}
 	return bParsedFiles;
 }
 
-bool FFeaturePackContentSource::ExtractListOfAdditionalFiles(const FString& ConfigFileAsString,TArray<FString>& FileList, bool& bContainsSource) const
+bool FFeaturePackContentSource::ExtractListOfAdditionalFiles(const FString& InConfigFileAsString,TArray<FString>& InFileList, bool& bContainsSource)
 {
 	FConfigFile PackConfig;
-	PackConfig.ProcessInputFileContents(ConfigFileAsString);
+	PackConfig.ProcessInputFileContents(InConfigFileAsString);
 	FConfigSection* AdditionalFilesSection = PackConfig.Find("AdditionalFilesToAdd");
 	
 	bContainsSource = false;
 	bool bParsedAdditionFiles = false;
 	if (AdditionalFilesSection)
 	{
+		TArray<FString> AdditionalFilesMap;
+		
 		bParsedAdditionFiles = true;
 		for (auto FilePair : *AdditionalFilesSection)
 		{
 			if (FilePair.Key.ToString().Contains("Files"))
 			{
-				FString Filename = FPaths::GetCleanFilename(FilePair.Value);
-				FString Directory = FPaths::RootDir() / FPaths::GetPath(FilePair.Value);
-				FPaths::MakeStandardFilename(Directory);
-				FPakFile::MakeDirectoryFromPath(Directory);
+				AdditionalFilesMap.Add(FilePair.Value);
+			}
+		}
+		BuildListOfAdditionalFiles(AdditionalFilesMap,InFileList,bContainsSource);
+				
+	}
+	return bParsedAdditionFiles;
+}
 
-				if (Filename.Contains(TEXT("*")))
-				{
-					TArray<FString> FoundFiles;
-					IFileManager::Get().FindFilesRecursive(FoundFiles, *Directory, *Filename, true, false);
-					FileList.Append(FoundFiles);
-					if (!bContainsSource)
-					{
-						for (const FString& FoundFile : FoundFiles)
-						{
-							if (FoundFile.StartsWith(TEXT("Source/")) || FoundFile.Contains(TEXT("/Source/")))
-							{
-								bContainsSource = true;
-								break;
-							}
-						}
-					}
-				}
-				else
+void FFeaturePackContentSource::BuildListOfAdditionalFiles(TArray<FString>& AdditionalFileSourceList,TArray<FString>& FileList, bool& bContainsSourceFiles)
+{
+	for (auto FileSource : AdditionalFileSourceList)
+	{
+		FString Filename = FPaths::GetCleanFilename(FileSource);
+		FString Directory = FPaths::RootDir() / FPaths::GetPath(FileSource);
+		FPaths::MakeStandardFilename(Directory);
+		FPakFile::MakeDirectoryFromPath(Directory);
+
+		if (Filename.Contains(TEXT("*")))
+		{
+			TArray<FString> FoundFiles;
+			IFileManager::Get().FindFilesRecursive(FoundFiles, *Directory, *Filename, true, false);
+			FileList.Append(FoundFiles);
+			if (!bContainsSourceFiles)
+			{
+				for (const FString& FoundFile : FoundFiles)
 				{
-					FileList.Add(Directory / Filename);
-					if (!bContainsSource && (FileList.Last().StartsWith(TEXT("Source/")) || FileList.Last().Contains(TEXT("/Source/"))))
+					if (FoundFile.StartsWith(TEXT("Source/")) || FoundFile.Contains(TEXT("/Source/")))
 					{
-						bContainsSource = true;
+						bContainsSourceFiles = true;
+						break;
 					}
 				}
 			}
 		}
+		else
+		{
+			FileList.Add(Directory / Filename);
+			if (!bContainsSourceFiles && (FileList.Last().StartsWith(TEXT("Source/")) || FileList.Last().Contains(TEXT("/Source/"))))
+			{
+				bContainsSourceFiles = true;
+			}
+		}
 	}
-	return bParsedAdditionFiles;
 }
 
 void FFeaturePackContentSource::ImportPendingPacks()
@@ -678,7 +635,7 @@ void FFeaturePackContentSource::RecordAndLogError(const FString& ErrorString)
 	ParseErrors.Add(ErrorString);
 }
 
-void FFeaturePackContentSource::CopyAdditionalFilesToFolder( const FString& DestinationFolder, TArray<FString>& FilesCopied, bool &bHasSourceFiles )
+void FFeaturePackContentSource::CopyAdditionalFilesToFolder( const FString& DestinationFolder, TArray<FString>& FilesCopied, bool &bHasSourceFiles, FString InGameFolder )
 {
 	FString ContentIdent = TEXT("Content/");
 	TArray<FString> FilesToAdd;
@@ -691,11 +648,17 @@ void FFeaturePackContentSource::CopyAdditionalFilesToFolder( const FString& Dest
 		if (ContentIndex != INDEX_NONE)
 		{
 			FString ContentFile = EachFile.RightChop(ContentIndex);
-			FString GameFolder = EachFile.Left(ContentIndex - 1);
-			int32 NameStart = GameFolder.Find(TEXT("/"), ESearchCase::IgnoreCase, ESearchDir::FromEnd);
-			GameFolder = GameFolder.RightChop(NameStart);
+			FString GameFolder = InGameFolder;
+			if(( GameFolder.IsEmpty() == false) && ( GameFolder.StartsWith(TEXT("/")) == false ))
+			{
+				GameFolder.InsertAt(0,TEXT("/"));
+			}
 			FPaths::NormalizeFilename(ContentFile);
-			ContentFile.InsertAt(ContentIdent.Len() - 1, GameFolder);
+			if( GameFolder.IsEmpty() == false)
+			{
+				ContentFile.InsertAt(ContentIdent.Len() - 1, GameFolder);
+			}
+
 			FString FinalDestination = DestinationFolder / ContentFile;
 			if (IFileManager::Get().Copy(*FinalDestination, *EachFile) == COPY_OK)
 			{
@@ -711,7 +674,7 @@ void FFeaturePackContentSource::CopyAdditionalFilesToFolder( const FString& Dest
 
 void FFeaturePackContentSource::InsertAdditionalFeaturePacks()
 {
-	// Now copy files in any additional packs (These are for the shared assets)
+	// copy files in any additional packs (These are for the shared assets)
 	EFeaturePackDetailLevel RequiredLevel = EFeaturePackDetailLevel::High;
 	for (int32 iExtraPack = 0; iExtraPack < AdditionalFeaturePacks.Num(); iExtraPack++)
 	{
@@ -726,10 +689,273 @@ void FFeaturePackContentSource::InsertAdditionalFeaturePacks()
 			{
 				bool bHasSourceFiles = false;
 				TArray<FString> FilesCopied;
-				NewContentSource->CopyAdditionalFilesToFolder(DestinationFolder, FilesCopied, bHasSourceFiles);
+				FString GameFolder = AdditionalFeaturePacks[iExtraPack].MountName;
+				NewContentSource->CopyAdditionalFilesToFolder(DestinationFolder, FilesCopied, bHasSourceFiles, GameFolder);
+			}
+		}
+	}
+}
+
+bool FFeaturePackContentSource::InsertAdditionalResources(TArray<FFeaturePackLevelSet> InAdditionalFeaturePacks,EFeaturePackDetailLevel RequiredLevel, const FString& InDestinationFolder,TArray<FString>& InFilesCopied )
+{
+	// Build a map of feature packs we have (This would probably be better elsewhere and stored in 2 arrays - one listing .upack packs and one non-upack packs)
+	IAddContentDialogModule& AddContentDialogModule = FModuleManager::LoadModuleChecked<IAddContentDialogModule>("AddContentDialog");
+	TMap<FString, FFeaturePackContentSource*> PackMap;
+	for (auto& ContentSourceProvider : *AddContentDialogModule.GetContentSourceProviderManager()->GetContentSourceProviders())
+	{
+		const TArray<TSharedRef<IContentSource>> ProviderSources = ContentSourceProvider->GetContentSources();
+		for (auto& EachSourceProvider : ProviderSources)
+		{
+			FFeaturePackContentSource* Source = (FFeaturePackContentSource*)&EachSourceProvider.Get();
+			FString ID = EachSourceProvider->GetIdent();
+			if( ID.IsEmpty() == false )
+			{
+				PackMap.Add(ID, Source);
+			}
+		}
+	}
+
+	int32 PacksInserted = 0;
+	for (int32 iExtraPack = 0; iExtraPack < InAdditionalFeaturePacks.Num(); iExtraPack++)
+	{		
+		// TODO - improve error handling here
+		FString PackName = InAdditionalFeaturePacks[iExtraPack].GetFeaturePackNameForLevel(RequiredLevel);
+		// If we have a 'non upack' pack, insert it else try a upack version
+		FFeaturePackContentSource* Insertable = PackMap.FindRef(PackName.Replace(TEXT(".upack"),TEXT("")));
+		TArray<FString> FilesCopied;
+		if( Insertable != nullptr )
+		{
+			bool bHasSourceFiles = false;
+			Insertable->CopyAdditionalFilesToFolder(InDestinationFolder, FilesCopied, bHasSourceFiles, InAdditionalFeaturePacks[iExtraPack].MountName);
+			PacksInserted++;
+		}
+		else
+		{
+			FString FullPath = FPaths::FeaturePackDir() + InAdditionalFeaturePacks[iExtraPack].GetFeaturePackNameForLevel(RequiredLevel);
+
+			if (FPlatformFileManager::Get().GetPlatformFile().FileExists(*FullPath))
+			{
+				FString InsertPath = InDestinationFolder + *InAdditionalFeaturePacks[iExtraPack].MountName;
+
+				TUniquePtr<FFeaturePackContentSource> NewContentSource = MakeUnique<FFeaturePackContentSource>(FullPath, true);
+				if (NewContentSource->IsDataValid() == true)
+				{
+					bool bHasSourceFiles = false;					
+					NewContentSource->CopyAdditionalFilesToFolder(InDestinationFolder, FilesCopied, bHasSourceFiles, InAdditionalFeaturePacks[iExtraPack].MountName);
+					PacksInserted++;
+				}
+			}
+		}
+		// Add to the list of files copied
+		InFilesCopied.Append(FilesCopied);
+	}
+	return PacksInserted == InAdditionalFeaturePacks.Num();
+}
+
+bool FFeaturePackContentSource::ParseManifestString(const FString& ManifestString)
+{
+	// Populate text fields from the manifest.
+	TSharedPtr<FJsonObject> ManifestObject;
+	TSharedRef<TJsonReader<>> ManifestReader = TJsonReaderFactory<>::Create(ManifestString);
+	FJsonSerializer::Deserialize(ManifestReader, ManifestObject);
+
+	if (ManifestReader->GetErrorMessage().IsEmpty() == false)
+	{
+		RecordAndLogError(FString::Printf(TEXT("Error in Feature pack %s. Failed to parse manifest: %s"), *FeaturePackPath, *ManifestReader->GetErrorMessage()));
+		Category = EContentSourceCategory::Unknown;
+		return false;
+	}
+
+	if (ManifestObject->HasTypedField<EJson::String>("Version") == true)
+	{
+		VersionNumber = ManifestObject->GetStringField("Version");
+	}
+	if (ManifestObject->HasTypedField<EJson::String>("Ident") == true)
+	{
+		Identity = ManifestObject->GetStringField("Ident");
+	}
+
+	TSharedPtr<FString> ManifestObjectErrorMessage;
+	if (TryValidateManifestObject(ManifestObject, ManifestObjectErrorMessage) == false)
+	{
+		RecordAndLogError(FString::Printf(TEXT("Error in Feature pack %s. Manifest object error: %s"), *FeaturePackPath, **ManifestObjectErrorMessage));
+		Category = EContentSourceCategory::Unknown;
+		return false;
+	}
+
+	for (TSharedPtr<FJsonValue> NameValue : ManifestObject->GetArrayField("Name"))
+	{
+		TSharedPtr<FJsonObject> LocalizedNameObject = NameValue->AsObject();
+		LocalizedNames.Add(FLocalizedText(
+			LocalizedNameObject->GetStringField("Language"),
+			FText::FromString(LocalizedNameObject->GetStringField("Text"))));
+	}
+
+	for (TSharedPtr<FJsonValue> DescriptionValue : ManifestObject->GetArrayField("Description"))
+	{
+		TSharedPtr<FJsonObject> LocalizedDescriptionObject = DescriptionValue->AsObject();
+		LocalizedDescriptions.Add(FLocalizedText(
+			LocalizedDescriptionObject->GetStringField("Language"),
+			FText::FromString(LocalizedDescriptionObject->GetStringField("Text"))));
+	}
+
+	// Parse asset types field
+	for (TSharedPtr<FJsonValue> AssetTypesValue : ManifestObject->GetArrayField("AssetTypes"))
+	{
+		TSharedPtr<FJsonObject> LocalizedAssetTypesObject = AssetTypesValue->AsObject();
+		LocalizedAssetTypesList.Add(FLocalizedText(
+			LocalizedAssetTypesObject->GetStringField("Language"),
+			FText::FromString(LocalizedAssetTypesObject->GetStringField("Text"))));
+	}
+
+	// Parse search tags field
+	if (ManifestObject->HasField("SearchTags") == true)
+	{
+		for (TSharedPtr<FJsonValue> AssetTypesValue : ManifestObject->GetArrayField("SearchTags"))
+		{
+			TSharedPtr<FJsonObject> LocalizedAssetTypesObject = AssetTypesValue->AsObject();
+			LocalizedSearchTags.Add(FLocalizedTextArray(
+				LocalizedAssetTypesObject->GetStringField("Language"),
+				LocalizedAssetTypesObject->GetStringField("Text")));
+		}
+	}
+
+	// Parse class types field
+	ClassTypes = ManifestObject->GetStringField("ClassTypes");
+
+	// Parse initial focus asset if we have one - this is not required
+	if (ManifestObject->HasTypedField<EJson::String>("FocusAsset") == true)
+	{
+		FocusAssetIdent = ManifestObject->GetStringField("FocusAsset");
+	}
+
+	// Use the path as the sort key - it will be alphabetical that way
+	SortKey = FeaturePackPath;
+	ManifestObject->TryGetStringField("SortKey", SortKey);
+
+	FString CategoryString = ManifestObject->GetStringField("Category");
+	UEnum* Enum = FindObjectChecked<UEnum>(ANY_PACKAGE, TEXT("EContentSourceCategory"));
+	int32 Index = Enum->FindEnumIndex(FName(*CategoryString));
+	Category = Index != INDEX_NONE ? (EContentSourceCategory)Index : EContentSourceCategory::Unknown;
+
+	// Thumbnail filename
+	IconFilename = ManifestObject->GetStringField("Thumbnail");
+
+	// Screenshots filenames
+	ScreenshotFilenameArray = ManifestObject->GetArrayField("Screenshots");
+	
+	// Parse additional files data
+	if (ManifestObject->HasTypedField<EJson::Object>("AdditionalFiles") == true)
+	{
+		TSharedPtr<FJsonObject> AdditionalFileObject = ManifestObject->GetObjectField("AdditionalFiles");
+		if( AdditionalFileObject->HasTypedField<EJson::String>("DestinationFilesFolder") == true )
+		{
+			AdditionalFilesForPack.DestinationFilesFolder = AdditionalFileObject->GetStringField("DestinationFilesFolder");
+			if (AdditionalFileObject->HasTypedField<EJson::Array>("AdditionalFilesList") == true)
+			{
+				for (TSharedPtr<FJsonValue> FileEntryValue : AdditionalFileObject->GetArrayField("AdditionalFilesList"))
+				{
+					const FString FileSpecString = FileEntryValue->AsString();
+					AdditionalFilesForPack.AdditionalFilesList.AddUnique(FileSpecString);
+				}
 			}
 		}
 	}
+
+	// Parse additional packs data
+	if (ManifestObject->HasTypedField<EJson::Array>("AdditionalFeaturePacks") == true)
+	{
+		UEnum* DetailEnum = FindObjectChecked<UEnum>(ANY_PACKAGE, TEXT("EFeaturePackDetailLevel"));
+		for (TSharedPtr<FJsonValue> AdditionalFeaturePackValue : ManifestObject->GetArrayField("AdditionalFeaturePacks"))
+		{
+			TSharedPtr<FJsonObject> EachAdditionalPack = AdditionalFeaturePackValue->AsObject();
+			FString MountName = EachAdditionalPack->GetStringField("MountName");
+
+			TArray<EFeaturePackDetailLevel>	Levels;
+			for (TSharedPtr<FJsonValue> DetailValue : EachAdditionalPack->GetArrayField("DetailLevels"))
+			{
+				const FString DetailString = DetailValue->AsString();
+				int32 eIndex = DetailEnum->FindEnumIndex(FName(*DetailString));
+				EFeaturePackDetailLevel EachLevel = Index != INDEX_NONE ? (EFeaturePackDetailLevel)eIndex : EFeaturePackDetailLevel::Standard;
+				Levels.AddUnique(EachLevel);
+			}
+			int32 PackIndex = AdditionalFeaturePacks.Add(FFeaturePackLevelSet(MountName,Levels));
+			// Check that the additional feature packs listed exist
+			FFeaturePackLevelSet& NewSet = AdditionalFeaturePacks[PackIndex];
+			for (int32 iLevel = 0; iLevel < NewSet.DetailLevels.Num(); iLevel++)
+			{
+				FString FullPath = FPaths::FeaturePackDir() + NewSet.GetFeaturePackNameForLevel(NewSet.DetailLevels[iLevel]);
+				if (!FPlatformFileManager::Get().GetPlatformFile().FileExists(*FullPath))
+				{
+					RecordAndLogError(FString::Printf(TEXT("Error in Feature pack %s. Cannot find additional pack %s."), *FeaturePackPath, *FullPath));
+				}
+			}
+		}
+	}
+	
+	bPackValid = true;
+	return true;
+}
+
+bool FFeaturePackContentSource::LoadFeaturePackImageData()
+{
+	TSharedPtr<TArray<uint8>> IconImageData = MakeShareable(new TArray<uint8>());
+ 	FString ThumbnailFile = FPaths::Combine(*MountPoint, TEXT("Media"), *IconFilename);
+ 	if (FFileHelper::LoadFileToArray( *IconImageData.Get(),*ThumbnailFile ) == true)
+ 	{
+ 		IconData = MakeShareable(new FImageData(IconFilename, IconImageData));
+ 	}
+ 	else
+ 	{
+ 		RecordAndLogError(FString::Printf(TEXT("Error in Feature pack %s. Cannot find thumbnail %s."), *FeaturePackPath, *ThumbnailFile));
+ 	}
+
+	// parse the screenshots field
+	for (const TSharedPtr<FJsonValue> ScreenshotFilename : ScreenshotFilenameArray)
+	{
+		TSharedPtr<TArray<uint8>> SingleScreenshotData = MakeShareable(new TArray<uint8>);
+ 		if (FFileHelper::LoadFileToArray(*SingleScreenshotData.Get(), *FPaths::Combine(*MountPoint, TEXT("Media"), *ScreenshotFilename->AsString()) ))
+		{
+			ScreenshotData.Add(MakeShareable(new FImageData(ScreenshotFilename->AsString(), SingleScreenshotData)));
+		}
+		else
+		{
+ 			RecordAndLogError(FString::Printf(TEXT("Error in Feature pack %s. Cannot find screenshot %s."), *FeaturePackPath, *ScreenshotFilename->AsString()));
+ 		}
+	}
+	return true;
+}
+
+bool FFeaturePackContentSource::LoadFeaturePackImageDataFromPackFile(FPakPlatformFile& PakPlatformFile)
+{
+	bool bAllImagesValid = true;
+	TSharedPtr<TArray<uint8>> IconImageData = MakeShareable(new TArray<uint8>());
+	FString ThumbnailFile = FPaths::Combine(*MountPoint, TEXT("Media"), *IconFilename);
+	if (LoadPakFileToBuffer(PakPlatformFile, ThumbnailFile, *IconImageData) == true)
+	{
+		IconData = MakeShareable(new FImageData(IconFilename, IconImageData));
+	}
+	else
+	{
+		RecordAndLogError(FString::Printf(TEXT("Error in Feature pack %s. Cannot find thumbnail %s."), *FeaturePackPath, *ThumbnailFile));
+		bAllImagesValid = false;
+	}
+
+	// parse the screenshots field
+	for (const TSharedPtr<FJsonValue> ScreenshotFilename : ScreenshotFilenameArray)
+	{
+		TSharedPtr<TArray<uint8>> SingleScreenshotData = MakeShareable(new TArray<uint8>);
+		if (LoadPakFileToBuffer(PakPlatformFile, FPaths::Combine(*MountPoint, TEXT("Media"), *ScreenshotFilename->AsString()), *SingleScreenshotData))
+		{
+			ScreenshotData.Add(MakeShareable(new FImageData(ScreenshotFilename->AsString(), SingleScreenshotData)));
+		}
+		else
+		{
+			RecordAndLogError(FString::Printf(TEXT("Error in Feature pack %s. Cannot find screenshot %s."), *FeaturePackPath, *ScreenshotFilename->AsString()));
+			bAllImagesValid = false;
+		}
+	}
+	return bAllImagesValid;
 }
 
 #undef LOCTEXT_NAMESPACE 
diff --git a/Engine/Source/Editor/AddContentDialog/Private/ContentSourceProviders/FeaturePack/FeaturePackContentSourceProvider.cpp b/Engine/Source/Editor/AddContentDialog/Private/ContentSourceProviders/FeaturePack/FeaturePackContentSourceProvider.cpp
--- a/Engine/Source/Editor/AddContentDialog/Private/ContentSourceProviders/FeaturePack/FeaturePackContentSourceProvider.cpp
+++ b/Engine/Source/Editor/AddContentDialog/Private/ContentSourceProviders/FeaturePack/FeaturePackContentSourceProvider.cpp
@@ -30,6 +30,7 @@ public:
 FFeaturePackContentSourceProvider::FFeaturePackContentSourceProvider()
 {
 	FeaturePackPath = FPaths::FeaturePackDir();
+	TemplatePath = FPaths::RootDir() + TEXT("Templates/");
 	StartUpDirectoryWatcher();
 	RefreshFeaturePacks();
 }
@@ -81,6 +82,8 @@ struct FFeaturePackCompareSortKey
 void FFeaturePackContentSourceProvider::RefreshFeaturePacks()
 {
 	ContentSources.Empty();
+	// TODO improve this and seperate handling of .upack and manifest (loose) parsing
+	// first the .upack files
 	IPlatformFile &PlatformFile = FPlatformFileManager::Get().GetPlatformFile();
 	FFillArrayDirectoryVisitor DirectoryVisitor;
 	PlatformFile.IterateDirectory( *FeaturePackPath, DirectoryVisitor );
@@ -95,6 +98,27 @@ void FFeaturePackContentSourceProvider::RefreshFeaturePacks()
 			}
 		}
 	}
+	
+	// Now the 'loose' feature packs
+	FFillArrayDirectoryVisitor TemplateDirectoryVisitor;
+	PlatformFile.IterateDirectoryRecursively(*TemplatePath, TemplateDirectoryVisitor);
+	for (auto TemplatePackFile: TemplateDirectoryVisitor.Files)
+	{
+		FString ThisPackRoot = FPaths::GetPath(TemplatePackFile);
+		if (ThisPackRoot.EndsWith(TEXT("FeaturePack")) == true)
+		{			
+			if(TemplatePackFile.EndsWith(TEXT("manifest.json")) == true)
+			{
+				TUniquePtr<FFeaturePackContentSource> NewContentSource = MakeUnique<FFeaturePackContentSource>(TemplatePackFile);
+				if (NewContentSource->IsDataValid())
+				{
+					ContentSources.Add(MakeShareable(NewContentSource.Release()));
+				}
+			}
+		}
+	}
+
+
 	ContentSources.Sort(FFeaturePackCompareSortKey());
 	OnContentSourcesChanged.ExecuteIfBound();
 }
diff --git a/Engine/Source/Editor/AddContentDialog/Private/ContentSourceProviders/FeaturePack/FeaturePackContentSourceProvider.h b/Engine/Source/Editor/AddContentDialog/Private/ContentSourceProviders/FeaturePack/FeaturePackContentSourceProvider.h
--- a/Engine/Source/Editor/AddContentDialog/Private/ContentSourceProviders/FeaturePack/FeaturePackContentSourceProvider.h
+++ b/Engine/Source/Editor/AddContentDialog/Private/ContentSourceProviders/FeaturePack/FeaturePackContentSourceProvider.h
@@ -30,6 +30,9 @@ private:
 	/** The path on disk to the directory containing the feature packs. */
 	FString FeaturePackPath;
 
+	/** The path on disk to the directory containing the templates folder. */
+	FString TemplatePath;
+
 	/** The delegate which gets called when the feature pack directory changes. This reference
 		is kept so that it can be unregistered correctly. */
 	IDirectoryWatcher::FDirectoryChanged DirectoryChangedDelegate;
diff --git a/Engine/Source/Editor/AddContentDialog/Private/IContentSource.h b/Engine/Source/Editor/AddContentDialog/Private/IContentSource.h
--- a/Engine/Source/Editor/AddContentDialog/Private/IContentSource.h
+++ b/Engine/Source/Editor/AddContentDialog/Private/IContentSource.h
@@ -118,5 +118,6 @@ public:
 
 	virtual ~IContentSource() { };
 
-
+	/** Gets the identity of the content. */
+	virtual FString GetIdent() const = 0;
 };
diff --git a/Engine/Source/Editor/AddContentDialog/Public/FeaturePackContentSource.h b/Engine/Source/Editor/AddContentDialog/Public/FeaturePackContentSource.h
--- a/Engine/Source/Editor/AddContentDialog/Public/FeaturePackContentSource.h
+++ b/Engine/Source/Editor/AddContentDialog/Public/FeaturePackContentSource.h
@@ -1,7 +1,7 @@
 // Copyright 1998-2015 Epic Games, Inc. All Rights Reserved.
 
 #pragma once
-#include "Private/IContentSource.h"
+#include "IContentSource.h"
 #include "FeaturePackContentSource.generated.h"
 
 class UObject;
@@ -106,14 +106,33 @@ struct FFeaturePackLevelSet
 	}
 };
 
+/* Structure that defines a shared feature pack resource. */
+USTRUCT()
+struct FFeatureAdditionalFiles
+{
+	GENERATED_BODY()
+
+	FFeatureAdditionalFiles(){};
+	
+	/* Name of the folder to insert the files to */
+	UPROPERTY()
+	FString DestinationFilesFolder;
+
+	/* List of files to insert */
+	UPROPERTY()
+	TArray<FString> AdditionalFilesList;
+};
+
 class FPakPlatformFile;
 struct FSearchEntry;
 
 /** A content source which represents a content upack. */
-class FFeaturePackContentSource : public IContentSource
+class ADDCONTENTDIALOG_API FFeaturePackContentSource : public IContentSource
 {
 public:
-	ADDCONTENTDIALOG_API FFeaturePackContentSource(FString InFeaturePackPath, bool bDontRegister = false);
+	FFeaturePackContentSource();
+	FFeaturePackContentSource(FString InFeaturePackPath, bool bDontRegister = false);
+
 	virtual ~FFeaturePackContentSource();
 
 	virtual TArray<FLocalizedText> GetLocalizedNames() const override;
@@ -126,10 +145,12 @@ public:
 	
 	virtual TSharedPtr<FImageData> GetIconData() const override;
 	virtual TArray<TSharedPtr<FImageData>> GetScreenshotData() const override;
+	FString GetFocusAssetName() const;
 
 	virtual bool InstallToProject(FString InstallPath) override;
 
 	void InsertAdditionalFeaturePacks();
+	bool InsertAdditionalResources(TArray<FFeaturePackLevelSet> InAdditionalFeaturePacks,EFeaturePackDetailLevel RequiredLevel, const FString& InDestinationFolder,TArray<FString>& InFilesCopied);
 
 	virtual bool IsDataValid() const override;
 	
@@ -145,7 +166,7 @@ public:
 	 * @param bContainsSource 	Set to true if the file list contains any source files
 	 * @returns true if config file was read and parsed successfully
 	 */
-	ADDCONTENTDIALOG_API void CopyAdditionalFilesToFolder( const FString& DestinationFolder, TArray<FString>& FilesCopied, bool &bHasSourceFiles );
+	void CopyAdditionalFilesToFolder( const FString& DestinationFolder, TArray<FString>& FilesCopied, bool &bHasSourceFiles, FString InGameFolder = FString() );
 
 	/*
 	 * Returns a list of additional files (including the path) as specified in the config file if one exists in the pack file.
@@ -154,16 +175,22 @@ public:
 	 * @param bContainsSource did the file list contain any source files
 	 * @returns true if config file was read and parsed successfully
 	 */
-	ADDCONTENTDIALOG_API bool GetAdditionalFilesForPack(TArray<FString>& FileList, bool& bContainsSource);
+	bool GetAdditionalFilesForPack(TArray<FString>& FileList, bool& bContainsSource);
 
-	static ADDCONTENTDIALOG_API void ImportPendingPacks();
+	static void ImportPendingPacks();
 	
+	/* Errors found when parsing manifest (if any) */
 	TArray<FString>	ParseErrors;
 	
+	void BuildListOfAdditionalFiles(TArray<FString>& AdditionalFileSourceList,TArray<FString>& FileList, bool& bContainsSourceFiles);
+
+	/* Get the identifier of the pack. */
+	virtual FString GetIdent() const;
+
 private:
 	static void ParseAndImportPacks();
 	bool LoadPakFileToBuffer(FPakPlatformFile& PakPlatformFile, FString Path, TArray<uint8>& Buffer);
-	FString GetFocusAssetName() const;
+	
 	
 	/*
 	 * Extract the list of additional files defined in config file to an array
@@ -172,27 +199,84 @@ private:
 	 * @param FileList		  array to receive list of files
 	 * @param bContainsSource did the file list contain any source files
 	 */
-	bool ExtractListOfAdditionalFiles(const FString& ConfigFile, TArray<FString>& FileList,bool& bContainsSource) const;
+	bool ExtractListOfAdditionalFiles(const FString& ConfigFile, TArray<FString>& FileList,bool& bContainsSource);
 
 	void RecordAndLogError(const FString& ErrorString);
+
+	/* Load the images for the icon and screen shots directly from disk */
+	bool LoadFeaturePackImageData();
+	
+	/* extract the images for the icon and screen shots from a pak file */
+	bool LoadFeaturePackImageDataFromPackFile(FPakPlatformFile& PakPlatformFile);
+	
+	/* Parse the manifest string describing this pack file */
+	bool ParseManifestString(const FString& ManifestString);
+
 	/** Selects an FLocalizedText from an array which matches either the supplied language code, or the default language code. */
 	FLocalizedTextArray ChooseLocalizedTextArray(TArray<FLocalizedTextArray> Choices, FString LanguageCode);
 	FLocalizedText ChooseLocalizedText(TArray<FLocalizedText> Choices, FString LanguageCode);
 	void TryAddFeaturePackCategory(FString CategoryTitle, TArray< TSharedPtr<FSearchEntry> >& OutSuggestions);
 
+	/* The path of the file we used to create this feature pack instance */
 	FString FeaturePackPath;
+	
+	/* Array of localised names */
 	TArray<FLocalizedText> LocalizedNames;
+	
+	/* Array of localised descriptions */
 	TArray<FLocalizedText> LocalizedDescriptions;
+	
+	/* Defines the type of feature pack this is */
 	EContentSourceCategory Category;
+	
+	/* Filename of the icon */
+	FString IconFilename;
+	
+	/* Image data for the icon */
 	TSharedPtr<FImageData> IconData;
+	
+	/* Filenames of the preview screenshots */
+	TArray<TSharedPtr<FJsonValue>> ScreenshotFilenameArray;
+	
+	/* Image data of the preview screenshots */
 	TArray<TSharedPtr<FImageData>> ScreenshotData;
+	
+	/* Array of localised assset type names */
 	TArray<FLocalizedText> LocalizedAssetTypesList;
+	
+	/* Comma delimited string listing the class types */
 	FString ClassTypes;
+	
+	/* true if the pack is valid */
 	bool bPackValid;
+	
+	/* Asset to focus after loading the pack */
 	FString FocusAssetIdent;
+	
+	/* Key used when sorting in the add dialog */
 	FString SortKey;
+	
+	/* Tags searched when typing in the super search box */
 	TArray<FLocalizedTextArray> LocalizedSearchTags;
+	
+	/* Other feature packs this pack needs (shared assets) */
 	TArray<FFeaturePackLevelSet> AdditionalFeaturePacks;
+	
+	/* Handle for the search clicked delegate handler */
 	FDelegateHandle SearchClickedHandle;
+	
+	/* Handle for the search changed delegate handler */
 	FDelegateHandle SearchChangedHandle;
+	
+	/* Additional files to copy when installing this pack */
+	FFeatureAdditionalFiles AdditionalFilesForPack;
+	
+	/* Are the contents in a pack file or did we just read a manifest for the pack */
+	bool bContentsInPakFile;
+
+	/* Feature pack mount point */
+	FString MountPoint;
+
+	FString Identity;
+	FString VersionNumber;
 };
\ No newline at end of file
diff --git a/Engine/Source/Editor/GameProjectGeneration/Private/GameProjectUtils.cpp b/Engine/Source/Editor/GameProjectGeneration/Private/GameProjectUtils.cpp
--- a/Engine/Source/Editor/GameProjectGeneration/Private/GameProjectUtils.cpp
+++ b/Engine/Source/Editor/GameProjectGeneration/Private/GameProjectUtils.cpp
@@ -3493,60 +3493,16 @@ bool GameProjectUtils::AddSharedContentToProject(const FProjectInformation &InPr
 		{
 			RequiredDetail = EFeaturePackDetailLevel::Standard;
 		}
-		
-		for (int32 iPack = 0; iPack < TemplateDefs->SharedContentPacks.Num(); ++iPack)
-		{
-			FFeaturePackLevelSet EachPack = TemplateDefs->SharedContentPacks[iPack];
-			EFeaturePackDetailLevel EachRequiredDetail = RequiredDetail;
-			
-			if (EachPack.DetailLevels.Num() == 1)
-			{
-				// If theres only only detail level override the requirement with that
-				EachRequiredDetail = EachPack.DetailLevels[0];
-			}
-			else if (EachPack.DetailLevels.Num() == 0) 
-			{
-				// We need at least one level !
-				FFormatNamedArguments Args;
-				Args.Add(TEXT("FullPackFilename"), FText::FromString(EachPack.MountName));
-				OutFailReason = FText::Format(LOCTEXT("NoLevelsDefined", "No detail levels defined in template pack '{FullPackFilename}'."), Args);
-				return false;
-			}
 
-			for (int32 iDetail = 0; iDetail < EachPack.DetailLevels.Num() ; iDetail++)
-			{
-				if (EachPack.DetailLevels[iDetail] == EachRequiredDetail)
-				{
-					// Build the Packname from the mount and detail
-					FString DetailString;
-					UEnum::GetValueAsString(TEXT("/Script/AddContentDialog.EFeaturePackDetailLevel"), EachRequiredDetail, DetailString);
-					
-					FString FullPackFilename = FPaths::FeaturePackDir() + EachPack.MountName +  DetailString + DefaultFeaturePackExtension;
-					if (FPaths::FileExists(FullPackFilename) == false)
-					{
-						FFormatNamedArguments Args;
-						Args.Add(TEXT("FullPackFilename"), FText::FromString(FullPackFilename));
-						OutFailReason = FText::Format(LOCTEXT("CantFindPack", "Cannot find template pack '{FullPackFilename}'."), Args);
-						return false;
-					}
-					TUniquePtr<FFeaturePackContentSource> NewContentSource = MakeUnique<FFeaturePackContentSource>(FullPackFilename, true);
-					if (NewContentSource->IsDataValid() == true)
-					{						
-						FString DestinationFolder = DestFolder;
-						FPaths::NormalizeDirectoryName(DestinationFolder);
-						bool bHasSourceFiles = false;
-						TArray<FString> FilesCopied;
-						NewContentSource->CopyAdditionalFilesToFolder(DestinationFolder, FilesCopied, bHasSourceFiles);						
-					}
-					else
-					{
-						FFormatNamedArguments Args;
-						Args.Add(TEXT("PackName"), FText::FromString(FullPackFilename));
-						OutFailReason = FText::Format(LOCTEXT("PackParseError", "Error parsing template pack '{PackName}'."), Args);
-						return false;
-					}
-				}
-			}
+
+		TUniquePtr<FFeaturePackContentSource> TempFeaturePack = MakeUnique<FFeaturePackContentSource>();
+		bool bCopied = TempFeaturePack->InsertAdditionalResources(TemplateDefs->SharedContentPacks,RequiredDetail, DestFolder,CreatedFiles);
+		if( bCopied == false )
+		{
+			FFormatNamedArguments Args;
+			Args.Add(TEXT("TemplateName"), FText::FromString(SrcFolder));
+			OutFailReason = FText::Format(LOCTEXT("SharedResourceError", "Error adding shared resources for '{TemplateName}'."), Args);
+			return false;		
 		}
 	}
 	return true;
diff --git a/Engine/Source/Programs/AutomationTool/Scripts/RocketBuild.Automation.cs b/Engine/Source/Programs/AutomationTool/Scripts/RocketBuild.Automation.cs
--- a/Engine/Source/Programs/AutomationTool/Scripts/RocketBuild.Automation.cs
+++ b/Engine/Source/Programs/AutomationTool/Scripts/RocketBuild.Automation.cs
@@ -66,24 +66,7 @@ namespace Rocket
 			"TP_VehicleAdv",
 			"TP_VehicleAdvBP",
 			"StarterContent",
-			"MobileStarterContent",
-			"GeometryHigh",		//Shared packs
-			"GeometryStandard",
-			"MannequinHigh",
-			"MannequinStandard",
-			"RollingHigh",
-			"RollingStandard",
-			"2DSideScrollerStandard",
-			"FirstPersonStandard",
-			"FlyingStandard",
-			"PuzzleStandard",
-			"RollingStandard",
-			"RollingHigh",
-			"SideScrollerStandard",
-			"ThirdPersonStandard",
-			"TwinStickStandard",
-			"VehicleStandard",
-			"VehicleAdvStandard",
+			"MobileStarterContent",			
 		};
 
 		public RocketBuild()
