Fixed shadowed var in FeaturePackContentSource.cpp

diff --git a/Engine/Source/Editor/AddContentDialog/Private/ContentSourceProviders/FeaturePack/FeaturePackContentSource.cpp b/Engine/Source/Editor/AddContentDialog/Private/ContentSourceProviders/FeaturePack/FeaturePackContentSource.cpp
--- a/Engine/Source/Editor/AddContentDialog/Private/ContentSourceProviders/FeaturePack/FeaturePackContentSource.cpp
+++ b/Engine/Source/Editor/AddContentDialog/Private/ContentSourceProviders/FeaturePack/FeaturePackContentSource.cpp
@@ -270,7 +270,6 @@ bool FFeaturePackContentSource::InstallToProject(FString InstallPath)
  		if (AdditionalFilesForPack.AdditionalFilesList.Num() != 0)
  		{
  			bool bHasSourceFiles = false;
- 			TArray<FString> FilesCopied;
  			CopyAdditionalFilesToFolder( FPaths::GameDir(), FilesCopied, bHasSourceFiles/*,AdditionalFilesFolder*/);
  		}
 
@@ -281,7 +280,6 @@ bool FFeaturePackContentSource::InstallToProject(FString InstallPath)
 			AssetPaths.Add(FeaturePackPath);
 			TArray<UObject*> ImportedObjects = AssetToolsModule.Get().ImportAssets(AssetPaths, InstallPath);
 
-
 			if (ImportedObjects.Num() == 0)
 			{
 				UE_LOG(LogFeaturePack, Warning, TEXT("No objects imported installing pack %s"), *InstallPath);
