fixed another shadow var in FeaturePackContentSource.cpp

diff --git a/Engine/Source/Editor/AddContentDialog/Private/ContentSourceProviders/FeaturePack/FeaturePackContentSource.cpp b/Engine/Source/Editor/AddContentDialog/Private/ContentSourceProviders/FeaturePack/FeaturePackContentSource.cpp
--- a/Engine/Source/Editor/AddContentDialog/Private/ContentSourceProviders/FeaturePack/FeaturePackContentSource.cpp
+++ b/Engine/Source/Editor/AddContentDialog/Private/ContentSourceProviders/FeaturePack/FeaturePackContentSource.cpp
@@ -466,8 +466,7 @@ bool FFeaturePackContentSource::GetAdditionalFilesForPack(TArray<FString>& FileL
 		{
 			// Create a pak platform file and mount the feature pack file.
 			FPakPlatformFile PakPlatformFile;
-			PakPlatformFile.Initialize(&FPlatformFileManager::Get().GetPlatformFile(), TEXT(""));
-			FString MountPoint = "root:/";
+			PakPlatformFile.Initialize(&FPlatformFileManager::Get().GetPlatformFile(), TEXT(""));			
 			PakPlatformFile.Mount(*FeaturePackPath, 0, *MountPoint);
 
 			// Gets the manifest file as a JSon string
