Removed feature pack warning about additional feature packs

diff --git a/Engine/Source/Editor/AddContentDialog/Private/ContentSourceProviders/FeaturePack/FeaturePackContentSource.cpp b/Engine/Source/Editor/AddContentDialog/Private/ContentSourceProviders/FeaturePack/FeaturePackContentSource.cpp
--- a/Engine/Source/Editor/AddContentDialog/Private/ContentSourceProviders/FeaturePack/FeaturePackContentSource.cpp
+++ b/Engine/Source/Editor/AddContentDialog/Private/ContentSourceProviders/FeaturePack/FeaturePackContentSource.cpp
@@ -877,16 +877,22 @@ bool FFeaturePackContentSource::ParseManifestString(const FString& ManifestStrin
 				Levels.AddUnique(EachLevel);
 			}
 			int32 PackIndex = AdditionalFeaturePacks.Add(FFeaturePackLevelSet(MountName,Levels));
-			// Check that the additional feature packs listed exist
-			FFeaturePackLevelSet& NewSet = AdditionalFeaturePacks[PackIndex];
-			for (int32 iLevel = 0; iLevel < NewSet.DetailLevels.Num(); iLevel++)
-			{
-				FString FullPath = FPaths::FeaturePackDir() + NewSet.GetFeaturePackNameForLevel(NewSet.DetailLevels[iLevel]);
-				if (!FPlatformFileManager::Get().GetPlatformFile().FileExists(*FullPath))
-				{
-					RecordAndLogError(FString::Printf(TEXT("Error in Feature pack %s. Cannot find additional pack %s."), *FeaturePackPath, *FullPath));
-				}
-			}
+// 			Check that the additional feature packs listed exist
+// 			FFeaturePackLevelSet& NewSet = AdditionalFeaturePacks[PackIndex];
+// 			for (int32 iLevel = 0; iLevel < NewSet.DetailLevels.Num(); iLevel++)
+// 			{
+// 				FString FullPath = FPaths::FeaturePackDir() + NewSet.GetFeaturePackNameForLevel(NewSet.DetailLevels[iLevel]);
+// 				if (!FPlatformFileManager::Get().GetPlatformFile().FileExists(*FullPath))
+// 				{
+// 					FString TemplatesFolder = TEXT("FeaturePack");
+// 					FString ThisTemplateRoot = FPaths::GetPath(FullPath);
+// 					if (ThisTemplateRoot.EndsWith(TemplatesFolder) == true)
+// 					{
+// 					}	
+// 			
+// 					RecordAndLogError(FString::Printf(TEXT("Error in Feature pack %s. Cannot find additional pack %s."), *FeaturePackPath, *FullPath));
+// 				}
+// 			}
 		}
 	}
 	
