Reskin "Add Content/Feature Pack" dialog.

diff --git a/Engine/Source/Editor/AddContentDialog/Private/AddContentDialogModule.cpp b/Engine/Source/Editor/AddContentDialog/Private/AddContentDialogModule.cpp
--- a/Engine/Source/Editor/AddContentDialog/Private/AddContentDialogModule.cpp
+++ b/Engine/Source/Editor/AddContentDialog/Private/AddContentDialogModule.cpp
@@ -7,7 +7,6 @@
 #include "Framework/Application/SlateApplication.h"
 #include "IAddContentDialogModule.h"
 #include "ContentSourceProviderManager.h"
-#include "AddContentDialogStyle.h"
 #include "SAddContentDialog.h"
 
 #include "ContentSourceProviders/FeaturePack/FeaturePackContentSourceProvider.h"
@@ -27,12 +26,10 @@ public:
 
 		ContentSourceProviderManager = TSharedPtr<FContentSourceProviderManager>(new FContentSourceProviderManager());
 		ContentSourceProviderManager->RegisterContentSourceProvider(MakeShareable(new FFeaturePackContentSourceProvider()));
-		FAddContentDialogStyle::Initialize();
 	}
 
 	virtual void ShutdownModule() override
 	{
-		FAddContentDialogStyle::Shutdown();
 		FWidgetCarouselModuleStyle::Shutdown();
 	}
 
diff --git a/Engine/Source/Editor/AddContentDialog/Private/AddContentDialogStyle.cpp b/Engine/Source/Editor/AddContentDialog/Private/AddContentDialogStyle.cpp
deleted file mode 100644
diff --git a/Engine/Source/Editor/AddContentDialog/Private/AddContentDialogStyle.h b/Engine/Source/Editor/AddContentDialog/Private/AddContentDialogStyle.h
deleted file mode 100644
diff --git a/Engine/Source/Editor/AddContentDialog/Private/SAddContentDialog.cpp b/Engine/Source/Editor/AddContentDialog/Private/SAddContentDialog.cpp
--- a/Engine/Source/Editor/AddContentDialog/Private/SAddContentDialog.cpp
+++ b/Engine/Source/Editor/AddContentDialog/Private/SAddContentDialog.cpp
@@ -9,7 +9,6 @@
 
 void SAddContentDialog::Construct(const FArguments& InArgs)
 {
-	SAssignNew(AddContentWidget, SAddContentWidget);
 
 	SWindow::Construct(SWindow::FArguments()
 		.Title(LOCTEXT("AddContentDialogTitle", "Add Content to the Project"))
@@ -19,16 +18,10 @@ void SAddContentDialog::Construct(const FArguments& InArgs)
 		.SupportsMaximize(false)
 		[
 			SNew(SBorder)
-			.BorderImage(FEditorStyle::GetBrush("ToolPanel.GroupBorder"))
-			.Padding(FMargin(15))
+			.BorderImage(FEditorStyle::GetBrush("Brushes.Panel"))
+			.Padding(FMargin(10,0))
 			[
-				SNew(SVerticalBox)
-
-				// Add content widget.
-				+ SVerticalBox::Slot()
-				[
-					AddContentWidget.ToSharedRef()
-				]
+				SNew(SAddContentWidget)
 			]
 		]);
 }
diff --git a/Engine/Source/Editor/AddContentDialog/Private/SAddContentDialog.h b/Engine/Source/Editor/AddContentDialog/Private/SAddContentDialog.h
--- a/Engine/Source/Editor/AddContentDialog/Private/SAddContentDialog.h
+++ b/Engine/Source/Editor/AddContentDialog/Private/SAddContentDialog.h
@@ -22,6 +22,4 @@ public:
 	void Construct(const FArguments& InArgs);
 
 private:
-	/** The widget representing available content and which content the user has selected. */
-	TSharedPtr<SAddContentWidget> AddContentWidget;
 };
diff --git a/Engine/Source/Editor/AddContentDialog/Private/SAddContentWidget.cpp b/Engine/Source/Editor/AddContentDialog/Private/SAddContentWidget.cpp
--- a/Engine/Source/Editor/AddContentDialog/Private/SAddContentWidget.cpp
+++ b/Engine/Source/Editor/AddContentDialog/Private/SAddContentWidget.cpp
@@ -7,15 +7,137 @@
 #include "Widgets/Input/SButton.h"
 #include "Widgets/Layout/SScrollBox.h"
 #include "Widgets/Input/SCheckBox.h"
+#include "Widgets/Input/SSegmentedControl.h"
 #include "EditorStyleSet.h"
-#include "AddContentDialogStyle.h"
-
 #include "Widgets/Input/SSearchBox.h"
-#include "WidgetCarouselStyle.h"
-#include "SWidgetCarouselWithNavigation.h"
+#include "Widgets/Layout/SSeparator.h"
+#include "Styling/StyleColors.h"
+#include "Internationalization/BreakIterator.h"
+#include "Framework/Application/SlateApplication.h"
+#include "WidgetCarousel/Public/WidgetCarouselStyle.h"
+#include "WidgetCarousel/Public/SWidgetCarouselWithNavigation.h"
+
 
 #define LOCTEXT_NAMESPACE "AddContentDialog"
 
+/** 
+ *
+ */
+class SGenericThumbnailTile : public SCompoundWidget
+{
+	SLATE_BEGIN_ARGS(SGenericThumbnailTile)
+		: _ThumbnailSize(102, 102)
+		, _ThumbnailPadding(FMargin(5.0f, 0))
+		, _IsSelected(false)
+	{}
+		SLATE_ATTRIBUTE(const FSlateBrush*, Image)
+		SLATE_ATTRIBUTE(FText, DisplayName)
+		/** The size of the thumbnail */
+		SLATE_ARGUMENT(FVector2D, ThumbnailSize)
+		/** The size of the image in the thumbnail. Used to make the thumbnail not fill the entire thumbnail area. If this is not set the Image brushes desired size is used */
+		SLATE_ARGUMENT(TOptional<FVector2D>, ImageSize)
+		SLATE_ARGUMENT(FMargin, ThumbnailPadding)
+		SLATE_ATTRIBUTE(bool, IsSelected)
+	SLATE_END_ARGS()
+
+	void Construct(const FArguments& InArgs)
+	{
+		Image = InArgs._Image;
+		Text = InArgs._DisplayName;
+		IsSelected = InArgs._IsSelected;
+
+		ChildSlot
+		[
+			SNew(SBorder)
+			.Padding(FMargin(0.0f, 0.0f, 5.0f, 5.0f))
+			.BorderImage(FAppStyle::Get().GetBrush("ProjectBrowser.ProjectTile.DropShadow"))
+			[
+				SNew(SOverlay)
+				+ SOverlay::Slot()
+				[
+					SNew(SVerticalBox)
+					// Thumbnail
+					+ SVerticalBox::Slot()
+					.AutoHeight()
+					.HAlign(HAlign_Center)
+					.VAlign(VAlign_Center)
+					[
+						SNew(SBox)
+						.WidthOverride(InArgs._ThumbnailSize.X)
+						.HeightOverride(InArgs._ThumbnailSize.Y)
+						[
+							SNew(SBorder)
+							.Padding(0)
+							.BorderImage(FAppStyle::Get().GetBrush("ProjectBrowser.ProjectTile.ThumbnailAreaBackground"))
+							.HAlign(HAlign_Center)
+							.VAlign(VAlign_Center)
+							[
+								SNew(SImage)
+								.Image(Image)
+								.DesiredSizeOverride(InArgs._ImageSize)
+							]
+						]
+					]
+					// Name
+					+ SVerticalBox::Slot()
+					[
+						SNew(SBorder)
+						.Padding(InArgs._ThumbnailPadding)
+						.VAlign(VAlign_Top)
+						.Padding(FMargin(3.0f, 3.0f))
+						.BorderImage(FAppStyle::Get().GetBrush("ProjectBrowser.ProjectTile.NameAreaBackground"))
+						[
+							SNew(STextBlock)
+							.Font(FAppStyle::Get().GetFontStyle("ProjectBrowser.ProjectTile.Font"))
+							.AutoWrapText(true)
+							.LineBreakPolicy(FBreakIterator::CreateWordBreakIterator())
+							.Text(InArgs._DisplayName)
+							.ColorAndOpacity(FAppStyle::Get().GetSlateColor("Colors.Foreground"))
+						]
+					]
+				]
+				+ SOverlay::Slot()
+				[
+					SNew(SImage)
+					.Visibility(EVisibility::HitTestInvisible)
+					.Image_Lambda
+					(
+						[this]()
+						{
+							const bool bSelected = IsSelected.Get();
+							const bool bHovered = IsHovered();
+
+							if (bSelected && bHovered)
+							{
+								static const FName SelectedHover("ProjectBrowser.ProjectTile.SelectedHoverBorder");
+								return FAppStyle::Get().GetBrush(SelectedHover);
+							}
+							else if (bSelected)
+							{
+								static const FName Selected("ProjectBrowser.ProjectTile.SelectedBorder");
+								return FAppStyle::Get().GetBrush(Selected);
+							}
+							else if (bHovered)
+							{
+								static const FName Hovered("ProjectBrowser.ProjectTile.HoverBorder");
+								return FAppStyle::Get().GetBrush(Hovered);
+							}
+
+							return FStyleDefaults::GetNoBrush();
+						}
+					)
+				]
+			]
+		];
+	}
+
+private:
+	TAttribute<const FSlateBrush*> Image;
+	TAttribute<FText> Text;
+	TAttribute<bool> IsSelected;
+};
+
+
 void SAddContentWidget::Construct(const FArguments& InArgs)
 {
 	ViewModel = FAddContentWidgetViewModel::CreateShared();
@@ -29,72 +151,90 @@ void SAddContentWidget::Construct(const FArguments& InArgs)
 	ChildSlot
 	[
 		SNew(SVerticalBox)
-
-		// Tab Buttons
 		+ SVerticalBox::Slot()
-		.AutoHeight()
-		.Padding(FMargin(10, 0))
 		[
-			SAssignNew(CategoryTabsContainer, SBox)
+			SNew(SHorizontalBox)
+			// Content Source Tiles
+			+ SHorizontalBox::Slot()
 			[
-				CreateCategoryTabs()
-			]
-		]
-
-		// Content Source Tab Page
-		+ SVerticalBox::Slot()
-		.FillHeight(3.0f)
-		[
-			SNew(SBorder)
-			.BorderImage(FAddContentDialogStyle::Get().GetBrush("AddContentDialog.TabBackground"))
-			.Padding(FMargin(15))
-			[
-				SNew(SHorizontalBox)
-
-				// Content Source Tiles
-				+ SHorizontalBox::Slot()
+				SNew(SVerticalBox)
+				+ SVerticalBox::Slot()
+				.AutoHeight()
+				.Padding(FMargin(0, 10, 0, 16))
+				.HAlign(HAlign_Center)
 				[
-					SNew(SVerticalBox)
-
-					// Content Source Filter
-					+ SVerticalBox::Slot()
-					.AutoHeight()
-					.Padding(FMargin(0, 0, 0, 5))
+					SAssignNew(CategoryTabsContainer, SBox)
 					[
-						SAssignNew(SearchBoxPtr, SSearchBox)
-						.OnTextChanged(this, &SAddContentWidget::SearchTextChanged)
-					]
-
-					// Content Source Tile View
-					+ SVerticalBox::Slot()
-					[
-						CreateContentSourceTileView()
+						CreateCategoryTabs()
 					]
 				]
+				// Content Source Filter
+				+ SVerticalBox::Slot()
+				.AutoHeight()
+				.Padding(FMargin(0, 0, 0, 8))
+				[
+					SAssignNew(SearchBoxPtr, SSearchBox)
+					.OnTextChanged(this, &SAddContentWidget::SearchTextChanged)
+				]
 
-				// Splitter
-				+ SHorizontalBox::Slot()
-				.Padding(FMargin(10, 0))
-				.AutoWidth()
+				// Content Source Tile View
+				+ SVerticalBox::Slot()
 				[
-					SNew(SBox)
-					.WidthOverride(2)
-					[
-						SNew(SImage)
-						.Image(FAddContentDialogStyle::Get().GetBrush("AddContentDialog.Splitter"))
-					]
+					CreateContentSourceTileView()
 				]
+			]
+
+			// Splitter
+			+ SHorizontalBox::Slot()
+			.Padding(FMargin(18, 0, 0, 0))
+			.AutoWidth()
+			[
+				SNew(SSeparator)
+				.Orientation(Orient_Vertical)
+				.Thickness(2.0f)
+			]
 
-				// Content Source Details
-				+ SHorizontalBox::Slot()
+			// Content Source Details
+			+ SHorizontalBox::Slot()
+			[
+				SAssignNew(ContentSourceDetailContainer, SBox)
 				[
-					SAssignNew(ContentSourceDetailContainer, SBox)
-					[
-						CreateContentSourceDetail(ViewModel->GetSelectedContentSource())
-					]
+					CreateContentSourceDetail(ViewModel->GetSelectedContentSource())
 				]
 			]
 		]
+		+ SVerticalBox::Slot()
+		.AutoHeight()
+		[
+			SNew(SSeparator)
+			.Thickness(2.0f)
+		]
+		+ SVerticalBox::Slot()
+		.AutoHeight()
+		.Padding(0, 16, 0, 16)
+		[
+			SNew(SHorizontalBox)
+			+ SHorizontalBox::Slot()
+			.VAlign(VAlign_Center)
+			.HAlign(HAlign_Right)
+			[
+				SNew(SButton)
+				.ButtonStyle(FEditorStyle::Get(), "PrimaryButton")
+				.TextStyle(FAppStyle::Get(), "DialogButtonText")
+				.OnClicked(this, &SAddContentWidget::AddButtonClicked)
+				.Text(LOCTEXT("AddToProjectButton", "Add to Project"))
+			]
+			+ SHorizontalBox::Slot()
+			.AutoWidth()
+			.VAlign(VAlign_Center)
+			.Padding(8.0f, 0.0f, 0.0f, 0.0f)
+			[
+				SNew(SButton)
+				.TextStyle(FAppStyle::Get(), "DialogButtonText")
+				.OnClicked(this, &SAddContentWidget::CancelButtonClicked)
+				.Text(LOCTEXT("Cancel", "Cancel"))
+			]
+		]
 	];
 }
 
@@ -105,42 +245,24 @@ const TArray<TSharedPtr<IContentSource>>* SAddContentWidget::GetContentSourcesTo
 
 TSharedRef<SWidget> SAddContentWidget::CreateCategoryTabs()
 {
-	TSharedRef<SHorizontalBox> TabBox = SNew(SHorizontalBox);
-	for (auto Category : *ViewModel->GetCategories())
+	TSharedRef<SSegmentedControl<FCategoryViewModel>> CategoriesWidget =
+		SNew(SSegmentedControl<FCategoryViewModel>)
+		.OnValueChanged(this, &SAddContentWidget::OnSelectedCategoryChanged)
+		.Value(this, &SAddContentWidget::GetSelectedCategory)
+		.TextStyle(FAppStyle::Get(), "DialogButtonText")
+		.MaxSegmentsPerLine(3);
+
+	const bool bRebuildChildren = false;
+	int32 Index = 0;
+	for (const FCategoryViewModel& Category : ViewModel->GetCategories())
 	{
-		TabBox->AddSlot()
-		.Padding(FMargin(0, 0, 5, 0))
-		.AutoWidth()
-		[
-			SNew(SCheckBox)
-			.Style(FAddContentDialogStyle::Get(), "AddContentDialog.CategoryTab")
-			.OnCheckStateChanged(this, &SAddContentWidget::CategoryCheckBoxCheckStateChanged, Category)
-			.IsChecked(this, &SAddContentWidget::GetCategoryCheckBoxCheckState, Category)
-			.Padding(FMargin(5))
-			[
-				SNew(SHorizontalBox)
+		CategoriesWidget->AddSlot(Category, bRebuildChildren)
+			.Text(Category.GetText());
+	}
 
-				+ SHorizontalBox::Slot()
-				.AutoWidth()
-				.VAlign(VAlign_Center)
-				.Padding(FMargin(0, 0, 5, 0))
-				[
-					SNew(SImage)
-					.Image(Category.GetIconBrush())
-				]
+	CategoriesWidget->RebuildChildren();
 
-				+ SHorizontalBox::Slot()
-				.AutoWidth()
-				.VAlign(VAlign_Center)
-				[
-					SNew(STextBlock)
-					.TextStyle(FAddContentDialogStyle::Get(), "AddContentDialog.HeadingText")
-					.Text(Category.GetText())
-				]
-			]
-		];
-	}
-	return TabBox;
+	return CategoriesWidget;
 }
 
 TSharedRef<SWidget> SAddContentWidget::CreateContentSourceTileView()
@@ -149,39 +271,30 @@ TSharedRef<SWidget> SAddContentWidget::CreateContentSourceTileView()
 	.ListItemsSource(ViewModel->GetContentSources())
 	.OnGenerateTile(this, &SAddContentWidget::CreateContentSourceIconTile)
 	.OnSelectionChanged(this, &SAddContentWidget::ContentSourceTileViewSelectionChanged)
-	.ItemWidth(70)
-	.ItemHeight(115)
+	.ClearSelectionOnClick(false)
+	.ItemAlignment(EListItemAlignment::LeftAligned)
+	.ItemWidth(102)
+	.ItemHeight(153)
 	.SelectionMode(ESelectionMode::Single);
+
 	ContentSourceTileView->SetSelection(ViewModel->GetSelectedContentSource(), ESelectInfo::Direct);
 	return ContentSourceTileView.ToSharedRef();
 }
 
 TSharedRef<ITableRow> SAddContentWidget::CreateContentSourceIconTile(TSharedPtr<FContentSourceViewModel> ContentSource, const TSharedRef<STableViewBase>& OwnerTable)
 {
-	return SNew(STableRow< TSharedPtr<FString> >, OwnerTable)
-	[
-		SNew(SVerticalBox)
+	TWeakPtr<FContentSourceViewModel> ContentSourceWeakPtr = ContentSource;
 
-		+ SVerticalBox::Slot()
-		.HAlign(EHorizontalAlignment::HAlign_Center)
-		.AutoHeight()
-		.Padding(FMargin(3))
+	return 
+		SNew(STableRow< TSharedPtr<FString> >, OwnerTable)
+		.Style(FAppStyle::Get(), "ProjectBrowser.TableRow")
+		.Padding(2.0f)
 		[
-			SNew(SImage)
+			SNew(SGenericThumbnailTile)
 			.Image(ContentSource->GetIconBrush().Get())
-		]
-
-		+ SVerticalBox::Slot()
-		.HAlign(HAlign_Center)
-		.AutoHeight()
-		.Padding(FMargin(3, 0, 3, 3))
-		[
-			SNew(STextBlock)
-			.Text(ContentSource->GetName())
-			.WrapTextAt(64)
-			.Justification(ETextJustify::Center)
-		]
-	];
+			.DisplayName(ContentSource->GetName())
+			.IsSelected_Lambda([ContentSourceWeakPtr, this] { return ViewModel->GetSelectedContentSource() == ContentSourceWeakPtr; })
+		];
 }
 
 TSharedRef<SWidget> SAddContentWidget::CreateContentSourceDetail(TSharedPtr<FContentSourceViewModel> ContentSource)
@@ -200,16 +313,18 @@ TSharedRef<SWidget> SAddContentWidget::CreateContentSourceDetail(TSharedPtr<FCon
 			]
 
 			+SScrollBox::Slot()
-			.Padding(FMargin(0, 0, 0, 5))
+			.Padding(FMargin(10, 0, 0, 5))
 			[
 				SNew(STextBlock)
-				.TextStyle(FAddContentDialogStyle::Get(), "AddContentDialog.HeadingText")
+				.TextStyle(FEditorStyle::Get(), "DialogButtonText")
+				.Font(FAppStyle::Get().GetFontStyle("HeadingExtraSmall"))
+				.ColorAndOpacity(FStyleColors::ForegroundHover)
 				.Text(ContentSource->GetName())
 				.AutoWrapText(true)
 			]
 
 			+ SScrollBox::Slot()
-			.Padding(FMargin(0, 0, 0, 5))
+			.Padding(FMargin(10, 0, 0, 5))
 			[
 				SNew(STextBlock)
 				.Text(ContentSource->GetDescription())
@@ -217,15 +332,15 @@ TSharedRef<SWidget> SAddContentWidget::CreateContentSourceDetail(TSharedPtr<FCon
 			]
 		
 			+ SScrollBox::Slot()
-			.Padding(FMargin(0, 0, 0, 5))
+			.Padding(FMargin(10, 0, 0, 0))
 			[
 				SNew(STextBlock)						
 				.Visibility(ContentSource->GetAssetTypes().IsEmpty() == false ? EVisibility::Visible : EVisibility::Collapsed)
-				.TextStyle(FAddContentDialogStyle::Get(), "AddContentDialog.HeadingText")
-				.Text(LOCTEXT("FeaturePackAssetReferences", "Asset types used in this pack:"))
+				.TextStyle(FEditorStyle::Get(), "DialogButtonText")
+				.Text(LOCTEXT("FeaturePackAssetReferences", "Asset Types Used"))
 			]
 			+ SScrollBox::Slot()
-			.Padding(FMargin(0, 0, 0, 5))
+			.Padding(FMargin(10, 0, 0, 5))
 			[
 				SNew(STextBlock)
 				.Text(ContentSource->GetAssetTypes())
@@ -234,15 +349,15 @@ TSharedRef<SWidget> SAddContentWidget::CreateContentSourceDetail(TSharedPtr<FCon
 			]
 		
 			+ SScrollBox::Slot()
-			.Padding(FMargin(0, 0, 0, 5))
+			.Padding(FMargin(10, 0, 0, 0))
 			[
 				SNew(STextBlock)
-				.TextStyle(FAddContentDialogStyle::Get(), "AddContentDialog.HeadingText")
+				.TextStyle(FEditorStyle::Get(), "DialogButtonText")
 				.Visibility(ContentSource->GetClassTypes().IsEmpty() == false ? EVisibility::Visible : EVisibility::Collapsed)
-				.Text(LOCTEXT("FeaturePackClassReferences", "Class types used in this pack:"))
+				.Text(LOCTEXT("FeaturePackClassReferences", "Class Types Used"))
 			] 
 			+ SScrollBox::Slot()
-			.Padding(FMargin(0, 0, 0, 5))
+			.Padding(FMargin(10, 0, 0, 5))
 			[
 				SNew(STextBlock)
 				.Text(FText::FromString(ContentSource->GetClassTypes()))
@@ -250,40 +365,6 @@ TSharedRef<SWidget> SAddContentWidget::CreateContentSourceDetail(TSharedPtr<FCon
 				.AutoWrapText(true)
 			]
 		];
-	
-		VerticalBox->AddSlot()
-		.AutoHeight()
-		.Padding(0, 10, 0, 0)
-		.HAlign(HAlign_Right)
-		[
-			SNew(SButton)
-			.ButtonStyle(FEditorStyle::Get(), "FlatButton.Success")
-			.OnClicked(this, &SAddContentWidget::AddButtonClicked)
-			.ContentPadding(FMargin(5, 5, 5, 5))
-			.HAlign(EHorizontalAlignment::HAlign_Center)
-			[
-				SNew(SHorizontalBox)
-
-				+ SHorizontalBox::Slot()
-				.VAlign(VAlign_Center)
-				.AutoWidth()
-				.Padding(0, 0, 2, 0)
-				[
-					SNew(STextBlock)
-					.TextStyle(FEditorStyle::Get(), "NormalText.Important")
-					.Font(FEditorStyle::Get().GetFontStyle("FontAwesome.10"))
-					.Text(FText::FromString(FString(TEXT("\xf067"))) /*fa-plus*/)
-				]
-
-				+ SHorizontalBox::Slot()
-				.AutoWidth()
-				[
-					SNew(STextBlock)
-					.TextStyle(FAddContentDialogStyle::Get(), "AddContentDialog.AddButton.TextStyle")
-					.Text(LOCTEXT("AddToProjectButton", "Add to Project"))
-				]
-			]
-		];
 	}
 	return VerticalBox;
 }
@@ -297,19 +378,6 @@ TSharedRef<SWidget> SAddContentWidget::CreateScreenshotCarousel(TSharedPtr<FCont
 		.WidgetItemsSource(ContentSource->GetScreenshotBrushes());
 }
 
-void SAddContentWidget::CategoryCheckBoxCheckStateChanged(ECheckBoxState CheckState, FCategoryViewModel Category)
-{
-	if (CheckState == ECheckBoxState::Checked)
-	{
-		ViewModel->SetSelectedCategory(Category);
-	}
-}
-
-ECheckBoxState SAddContentWidget::GetCategoryCheckBoxCheckState(FCategoryViewModel Category) const
-{
-	return (Category == ViewModel->GetSelectedCategory()) ? ECheckBoxState::Checked : ECheckBoxState::Unchecked;
-}
-
 void SAddContentWidget::SearchTextChanged(const FText& SearchText)
 {
 	ViewModel->SetSearchText(SearchText);
@@ -330,6 +398,25 @@ FReply SAddContentWidget::AddButtonClicked()
 	return FReply::Handled();
 }
 
+FReply SAddContentWidget::CancelButtonClicked()
+{
+	TSharedPtr<SWindow> MyWindow = FSlateApplication::Get().FindWidgetWindow(AsShared());
+
+	MyWindow->RequestDestroyWindow();
+
+	return FReply::Handled();
+}
+
+void SAddContentWidget::OnSelectedCategoryChanged(FCategoryViewModel SelectedCategory)
+{
+	ViewModel->SetSelectedCategory(SelectedCategory);
+}
+
+FCategoryViewModel SAddContentWidget::GetSelectedCategory() const
+{
+	return ViewModel->GetSelectedCategory();
+}
+
 TSharedRef<SWidget> SAddContentWidget::CreateScreenshotWidget(TSharedPtr<FSlateBrush> ScreenshotBrush)
 {
 	return SNew(SImage)
diff --git a/Engine/Source/Editor/AddContentDialog/Private/SAddContentWidget.h b/Engine/Source/Editor/AddContentDialog/Private/SAddContentWidget.h
--- a/Engine/Source/Editor/AddContentDialog/Private/SAddContentWidget.h
+++ b/Engine/Source/Editor/AddContentDialog/Private/SAddContentWidget.h
@@ -54,13 +54,6 @@ private:
 	/** Creates the widget the displays a screenshot in the screenshot carousel. */
 	TSharedRef<SWidget> CreateScreenshotWidget(TSharedPtr<FSlateBrush> ScreenshotBrush);
 
-
-	/** Handles the user clicking on one of the check boxes representing the category tabs. */
-	void CategoryCheckBoxCheckStateChanged(ECheckBoxState CheckState, FCategoryViewModel Category);
-
-	/** Gets the check state for one of the check boxes representing the category tabs. */
-	ECheckBoxState GetCategoryCheckBoxCheckState(FCategoryViewModel Category) const;
-
 	/** Handles the text in the search box changing. */
 	void SearchTextChanged(const FText& SearchText);
 
@@ -70,6 +63,14 @@ private:
 	/** Handles the add content to project button being clicked. */
 	FReply AddButtonClicked();
 
+	/** Handles the close button being clicked. */
+	FReply CancelButtonClicked();
+
+	/** Handles selected category changing in the UI */
+	void OnSelectedCategoryChanged(FCategoryViewModel SelectedCategory);
+
+	/** Gets the curerently selected category */
+	FCategoryViewModel GetSelectedCategory() const;
 
 	/** Handles the available categories changing on the view model. */
 	void CategoriesChanged();
diff --git a/Engine/Source/Editor/AddContentDialog/Private/ViewModels/AddContentWidgetViewModel.cpp b/Engine/Source/Editor/AddContentDialog/Private/ViewModels/AddContentWidgetViewModel.cpp
--- a/Engine/Source/Editor/AddContentDialog/Private/ViewModels/AddContentWidgetViewModel.cpp
+++ b/Engine/Source/Editor/AddContentDialog/Private/ViewModels/AddContentWidgetViewModel.cpp
@@ -35,9 +35,9 @@ void FAddContentWidgetViewModel::Initialize()
 	BuildContentSourceViewModels();
 }
 
-const TArray<FCategoryViewModel>* FAddContentWidgetViewModel::GetCategories()
+const TArray<FCategoryViewModel>& FAddContentWidgetViewModel::GetCategories() const
 {
-	return &Categories;
+	return Categories;
 }
 
 void FAddContentWidgetViewModel::SetOnCategoriesChanged(FOnCategoriesChanged OnCategoriesChangedIn)
diff --git a/Engine/Source/Editor/AddContentDialog/Private/ViewModels/AddContentWidgetViewModel.h b/Engine/Source/Editor/AddContentDialog/Private/ViewModels/AddContentWidgetViewModel.h
--- a/Engine/Source/Editor/AddContentDialog/Private/ViewModels/AddContentWidgetViewModel.h
+++ b/Engine/Source/Editor/AddContentDialog/Private/ViewModels/AddContentWidgetViewModel.h
@@ -23,7 +23,7 @@ public:
 	static TSharedRef<FAddContentWidgetViewModel> CreateShared();
 
 	/** Gets the view models for the current set of content source categories. */
-	const TArray<FCategoryViewModel>* GetCategories();
+	const TArray<FCategoryViewModel>& GetCategories() const;
 
 	/** Sets the delegate which should be executed when the set of categories changes. */
 	void SetOnCategoriesChanged(FOnCategoriesChanged OnCategoriesChangedIn);
diff --git a/Engine/Source/Editor/AddContentDialog/Private/ViewModels/CategoryViewModel.cpp b/Engine/Source/Editor/AddContentDialog/Private/ViewModels/CategoryViewModel.cpp
--- a/Engine/Source/Editor/AddContentDialog/Private/ViewModels/CategoryViewModel.cpp
+++ b/Engine/Source/Editor/AddContentDialog/Private/ViewModels/CategoryViewModel.cpp
@@ -2,7 +2,6 @@
 
 #include "ViewModels/CategoryViewModel.h"
 #include "IContentSource.h"
-#include "AddContentDialogStyle.h"
 
 #define LOCTEXT_NAMESPACE "ContentSourceViewModel"
 
@@ -23,11 +22,6 @@ FText FCategoryViewModel::GetText() const
 	return Text;
 }
 
-const FSlateBrush* FCategoryViewModel::GetIconBrush() const
-{
-	return IconBrush;
-}
-
 uint32 FCategoryViewModel::GetTypeHash() const
 {
 	return (uint32)Category;
@@ -38,33 +32,27 @@ void FCategoryViewModel::Initialize()
 	switch (Category)
 	{
 	case EContentSourceCategory::BlueprintFeature:
-		Text = LOCTEXT("BlueprintFeature", "Blueprint Feature");
-		IconBrush = FAddContentDialogStyle::Get().GetBrush("AddContentDialog.BlueprintFeatureCategory");
+		Text = LOCTEXT("BlueprintFeature", "Blueprint");
 		SortID = 0;
 		break;
 	case EContentSourceCategory::CodeFeature:
-		Text = LOCTEXT("CodeFeature", "C++ Feature");
-		IconBrush = FAddContentDialogStyle::Get().GetBrush("AddContentDialog.CodeFeatureCategory");
+		Text = LOCTEXT("CodeFeature", "C++");
 		SortID = 1;
 		break;
 	case EContentSourceCategory::EnterpriseFeature:
 		Text = LOCTEXT("EnterpriseFeature", "Unreal Studio Feature");
-		IconBrush = FAddContentDialogStyle::Get().GetBrush("AddContentDialog.BlueprintFeatureCategory");
 		SortID = 2;
 		break;
 	case EContentSourceCategory::Content:
-		Text = LOCTEXT("ContentPacks", "Content Packs");
-		IconBrush = FAddContentDialogStyle::Get().GetBrush("AddContentDialog.ContentPackCategory");
+		Text = LOCTEXT("ContentPacks", "Content");
 		SortID = 3;
 		break;
 	case EContentSourceCategory::EnterpriseContent:
-		Text = LOCTEXT("EnterpriseContentPacks", "Unreal Studio Content Packs");
-		IconBrush = FAddContentDialogStyle::Get().GetBrush("AddContentDialog.ContentPackCategory");
+		Text = LOCTEXT("EnterpriseContentPacks", "Unreal Studio Content");
 		SortID = 4;
 		break;
 	default:
 		Text = LOCTEXT("Miscellaneous", "Miscellaneous");
-		IconBrush = FAddContentDialogStyle::Get().GetBrush("AddContentDialog.UnknownCategory");
 		SortID = 5;
 		break;
 	}
diff --git a/Engine/Source/Editor/AddContentDialog/Private/ViewModels/CategoryViewModel.h b/Engine/Source/Editor/AddContentDialog/Private/ViewModels/CategoryViewModel.h
--- a/Engine/Source/Editor/AddContentDialog/Private/ViewModels/CategoryViewModel.h
+++ b/Engine/Source/Editor/AddContentDialog/Private/ViewModels/CategoryViewModel.h
@@ -18,12 +18,9 @@ public:
 	/** Gets the display name of the category. */
 	FText GetText() const;
 
-	/** Gets the brush which should be used to draw the icon for the category. */
-	const FSlateBrush* GetIconBrush() const;
-
 	inline bool operator==(const FCategoryViewModel& Other) const
 	{
-		return Other.Text.EqualTo(Text) && (Other.IconBrush == IconBrush);
+		return Category == Other.Category;
 	}
 
 	bool operator<(FCategoryViewModel const& Other) const
@@ -39,7 +36,6 @@ private:
 private:
 	EContentSourceCategory Category;
 	FText Text;
-	const FSlateBrush* IconBrush;
 	int SortID;
 };
 
diff --git a/Engine/Source/Runtime/Slate/Public/Widgets/Input/SSegmentedControl.h b/Engine/Source/Runtime/Slate/Public/Widgets/Input/SSegmentedControl.h
--- a/Engine/Source/Runtime/Slate/Public/Widgets/Input/SSegmentedControl.h
+++ b/Engine/Source/Runtime/Slate/Public/Widgets/Input/SSegmentedControl.h
@@ -77,23 +77,30 @@ public:
 	DECLARE_DELEGATE_OneParam( FOnValueChanged, OptionType );
 
 	SLATE_BEGIN_ARGS( SSegmentedControl<OptionType> )
-	: _Style(&FAppStyle::Get().GetWidgetStyle<FSegmentedControlStyle>("SegmentedControl"))
+		: _Style(&FAppStyle::Get().GetWidgetStyle<FSegmentedControlStyle>("SegmentedControl"))
+		, _TextStyle(&FAppStyle::Get().GetWidgetStyle<FTextBlockStyle>("ButtonText"))
+		, _MaxSegmentsPerLine(0)
 	{}
 		/** Slot type supported by this panel */
 		SLATE_SUPPORTS_SLOT(FSlot<OptionType>)
 
 		/** Styling for this control */
 		SLATE_STYLE_ARGUMENT(FSegmentedControlStyle, Style)
+	
+		/** Styling for the text in each slot. If a custom widget is supplied for a slot this argument is not used */
+		SLATE_STYLE_ARGUMENT(FTextBlockStyle, TextStyle)
 
 		/** The current control value */
-		SLATE_ATTRIBUTE( OptionType, Value )
+		SLATE_ATTRIBUTE(OptionType, Value)
 
 		/** Padding to apply to each slot */
 		SLATE_ATTRIBUTE(FMargin, UniformPadding)
 
 		/** Called when the value is changed */
-		SLATE_EVENT( FOnValueChanged, OnValueChanged )
+		SLATE_EVENT(FOnValueChanged, OnValueChanged)
 		
+		/** Optional maximum number of segments per line before the control wraps vertically to the next line. If this value is <= 0 no wrapping happens */
+		SLATE_ARGUMENT(int32, MaxSegmentsPerLine)
 	SLATE_END_ARGS()
 
 	void Construct( const FArguments& InArgs )
@@ -101,12 +108,15 @@ public:
 		check(InArgs._Style);
 
 		Style = InArgs._Style;
+		TextStyle = InArgs._TextStyle;
 
 		CurrentValue = InArgs._Value;
 		OnValueChanged = InArgs._OnValueChanged;
 
 		UniformPadding = InArgs._UniformPadding;
 
+		MaxSegmentsPerLine = InArgs._MaxSegmentsPerLine;
+
 		const int32 NumSlots = InArgs.Slots.Num();
 		Children.Reserve(NumSlots);
 		for (int32 SlotIndex = 0; SlotIndex < NumSlots; ++SlotIndex)
@@ -154,20 +164,22 @@ public:
 				.AutoWidth()
 				[
 					SNew(STextBlock)
-					.TextStyle( &FAppStyle::Get().GetWidgetStyle< FTextBlockStyle >( "ButtonText" ))
+					.TextStyle(TextStyle)
 					.Justification(ETextJustify::Center)
 					.Text_Lambda([ChildSlotPtr] () { return ChildSlotPtr->_Text.Get(); } ) 
 				];
 			}
 
-			UniformBox->AddSlot(SlotIndex, 0)
+			int32 ColumnIndex = MaxSegmentsPerLine ? SlotIndex % MaxSegmentsPerLine : SlotIndex;
+			UniformBox->AddSlot(ColumnIndex, MaxSegmentsPerLine > 0 ? SlotIndex / MaxSegmentsPerLine : 0)
 			// Note HAlignment is applied at the check box level because if it were applied here it would make the slots look physically disconnected from each other 
 			.VAlign(ChildSlotPtr->VAlignment)
 			[
 				SNew(SCheckBox)
+				.Clipping(EWidgetClipping::ClipToBounds)
 				.HAlign(ChildSlotPtr->HAlignment)
 				.ToolTipText(ChildSlotPtr->_Tooltip)
-				.Style(SlotIndex == 0 ? &Style->FirstControlStyle : SlotIndex == (NumSlots - 1) ? &Style->LastControlStyle : &Style->ControlStyle)
+				.Style(ColumnIndex == 0 ? &Style->FirstControlStyle : ColumnIndex == (NumSlots - 1) ? &Style->LastControlStyle : &Style->ControlStyle)
 				.IsChecked(this, &SSegmentedControl::IsCurrentValue, ChildValue)
 				.OnCheckStateChanged(this, &SSegmentedControl::CommitValue, ChildValue)
 				.Padding(UniformPadding)
@@ -185,13 +197,16 @@ public:
 	}
 
 	// Slot Management
-	FSlot<OptionType>& AddSlot(const OptionType& InValue)
+	FSlot<OptionType>& AddSlot(const OptionType& InValue, bool bRebuildChildren = true)
 	{
 		FSlot<OptionType>& NewSlot = *(new FSlot<OptionType>(InValue));
 
 		Children.Add( &NewSlot );
-		RebuildChildren();
 
+		if (bRebuildChildren)
+		{
+			RebuildChildren();
+		}
 		return NewSlot;
 	}
 
@@ -237,4 +252,8 @@ protected:
 	TAttribute<FMargin> UniformPadding;
 
 	const FSegmentedControlStyle* Style;
+
+	const FTextBlockStyle* TextStyle;
+
+	int32 MaxSegmentsPerLine = 0;
 };
