PR #7931: Stops LogFeaturePack: Warning when a new project is opened. (Contributed by [REDACTED_PERSON])

diff --git a/Engine/Source/Editor/AddContentDialog/Private/ContentSourceProviders/FeaturePack/FeaturePackContentSource.cpp b/Engine/Source/Editor/AddContentDialog/Private/ContentSourceProviders/FeaturePack/FeaturePackContentSource.cpp
--- a/Engine/Source/Editor/AddContentDialog/Private/ContentSourceProviders/FeaturePack/FeaturePackContentSource.cpp
+++ b/Engine/Source/Editor/AddContentDialog/Private/ContentSourceProviders/FeaturePack/FeaturePackContentSource.cpp
@@ -554,7 +554,7 @@ void FFeaturePackContentSource::ParseAndImportPacks()
 			}
 		}
 	}
-	UE_LOG(LogFeaturePack, Warning, TEXT("Inserted %d feature packs"), PacksInserted++);
+	UE_LOG(LogFeaturePack, Log, TEXT("Inserted %d feature packs"), PacksInserted);
 }
 
 void FFeaturePackContentSource::RecordAndLogError(const FString& ErrorString)
