Feature packs can now have multiple categories specified for them in the "Category" field of the manifest.json file.

Fixed several standards-breaking uses of auto and lots of const-incorrect code.

diff --git a/Engine/Source/Editor/AddContentDialog/Private/ContentSourceProviders/FeaturePack/FeaturePackContentSource.cpp b/Engine/Source/Editor/AddContentDialog/Private/ContentSourceProviders/FeaturePack/FeaturePackContentSource.cpp
--- a/Engine/Source/Editor/AddContentDialog/Private/ContentSourceProviders/FeaturePack/FeaturePackContentSource.cpp
+++ b/Engine/Source/Editor/AddContentDialog/Private/ContentSourceProviders/FeaturePack/FeaturePackContentSource.cpp
@@ -112,7 +112,7 @@ bool TryValidateManifestObject(TSharedPtr<FJsonObject> ManifestObject, TSharedPt
 		return false;
 	}
 
-	if (ManifestObject->HasTypedField<EJson::String>("Category") == false)
+	if (ManifestObject->HasTypedField<EJson::Array>("Category") == false && ManifestObject->HasTypedField<EJson::String>("Category") == false)
 	{
 		ErrorMessage = MakeShareable(new FString("Manifest object missing 'Category' field"));
 		return false;
@@ -154,7 +154,7 @@ FFeaturePackContentSource::FFeaturePackContentSource(FString InFeaturePackPath)
 	bPackValid = false;
 	
 	FString ManifestString;
-	Category = EContentSourceCategory::Unknown;
+	Categories = { EContentSourceCategory::Unknown };
 	if( InFeaturePackPath.EndsWith(TEXT(".upack") ) == true )
 	{
 		bContentsInPakFile = true;
@@ -216,30 +216,30 @@ bool FFeaturePackContentSource::LoadPakFileToBuffer(FPakPlatformFile& PakPlatfor
 	return bResult;
 }
 
-TArray<FLocalizedText> FFeaturePackContentSource::GetLocalizedNames() const
+const TArray<FLocalizedText>& FFeaturePackContentSource::GetLocalizedNames() const
 {
 	return LocalizedNames;
 }
 
-TArray<FLocalizedText> FFeaturePackContentSource::GetLocalizedDescriptions() const
+const TArray<FLocalizedText>& FFeaturePackContentSource::GetLocalizedDescriptions() const
 {
 	return LocalizedDescriptions;
 }
 
-TArray<FLocalizedText> FFeaturePackContentSource::GetLocalizedAssetTypes() const
+const TArray<FLocalizedText>& FFeaturePackContentSource::GetLocalizedAssetTypes() const
 {
 	return LocalizedAssetTypesList;
 }
 
 
-FString FFeaturePackContentSource::GetClassTypesUsed() const
+const FString& FFeaturePackContentSource::GetClassTypesUsed() const
 {
 	return ClassTypes;
 }
 
-EContentSourceCategory FFeaturePackContentSource::GetCategory() const
+const TArray<EContentSourceCategory>& FFeaturePackContentSource::GetCategories() const
 {
-	return Category;
+	return Categories;
 }
 
 TSharedPtr<FImageData> FFeaturePackContentSource::GetIconData() const
@@ -247,7 +247,7 @@ TSharedPtr<FImageData> FFeaturePackContentSource::GetIconData() const
 	return IconData;
 }
 
-TArray<TSharedPtr<FImageData>> FFeaturePackContentSource::GetScreenshotData() const
+const TArray<TSharedPtr<FImageData>>& FFeaturePackContentSource::GetScreenshotData() const
 {
 	return ScreenshotData;
 }
@@ -287,7 +287,7 @@ bool FFeaturePackContentSource::InstallToProject(FString InstallPath)
 			{
 				// Save any imported assets.
 				TArray<UPackage*> ToSave;
-				for (auto ImportedObject : ImportedObjects)
+				for (UObject* ImportedObject : ImportedObjects)
 				{
 					ToSave.AddUnique(ImportedObject->GetOutermost());
 				}
@@ -328,17 +328,17 @@ bool FFeaturePackContentSource::IsDataValid() const
 	return true;	
 }
 
-FString FFeaturePackContentSource::GetFocusAssetName() const
+const FString& FFeaturePackContentSource::GetFocusAssetName() const
 {
 	return FocusAssetIdent;
 }
 
-FString FFeaturePackContentSource::GetSortKey() const
+const FString& FFeaturePackContentSource::GetSortKey() const
 {
 	return SortKey;
 }
 
-FString FFeaturePackContentSource::GetIdent() const
+const FString& FFeaturePackContentSource::GetIdent() const
 {
 	return Identity;
 }
@@ -430,7 +430,7 @@ bool FFeaturePackContentSource::ExtractListOfAdditionalFiles(const FString& InCo
 		TArray<FString> AdditionalFilesMap;
 		
 		bParsedAdditionFiles = true;
-		for (auto FilePair : *AdditionalFilesSection)
+		for (auto& FilePair : *AdditionalFilesSection)
 		{
 			if (FilePair.Key.ToString().Contains("Files"))
 			{
@@ -443,9 +443,9 @@ bool FFeaturePackContentSource::ExtractListOfAdditionalFiles(const FString& InCo
 	return bParsedAdditionFiles;
 }
 
-void FFeaturePackContentSource::BuildListOfAdditionalFiles(TArray<FString>& AdditionalFileSourceList,TArray<FString>& FileList, bool& bContainsSourceFiles)
+void FFeaturePackContentSource::BuildListOfAdditionalFiles(TArray<FString>& AdditionalFileSourceList, TArray<FString>& FileList, bool& bContainsSourceFiles)
 {
-	for (auto FileSource : AdditionalFileSourceList)
+	for (const FString& FileSource : AdditionalFileSourceList)
 	{
 		FString Filename = FPaths::GetCleanFilename(FileSource);
 		FString Directory = FPaths::RootDir() / FPaths::GetPath(FileSource);
@@ -545,7 +545,7 @@ void FFeaturePackContentSource::ParseAndImportPacks()
 			{
 				// Save any imported assets.
 				TArray<UPackage*> ToSave;
-				for (auto ImportedObject : EachPackData.ImportedObjects)
+				for (UObject* ImportedObject : EachPackData.ImportedObjects)
 				{
 					ToSave.AddUnique(ImportedObject->GetOutermost());
 				}
@@ -629,16 +629,16 @@ bool FFeaturePackContentSource::InsertAdditionalResources(TArray<FFeaturePackLev
 	// Build a map of feature packs we have (This would probably be better elsewhere and stored in 2 arrays - one listing .upack packs and one non-upack packs)
 	IAddContentDialogModule& AddContentDialogModule = FModuleManager::LoadModuleChecked<IAddContentDialogModule>("AddContentDialog");
 	TMap<FString, FFeaturePackContentSource*> PackMap;
-	for (auto& ContentSourceProvider : *AddContentDialogModule.GetContentSourceProviderManager()->GetContentSourceProviders())
+	for (const TSharedRef<IContentSourceProvider>& ContentSourceProvider : *AddContentDialogModule.GetContentSourceProviderManager()->GetContentSourceProviders())
 	{
 		const TArray<TSharedRef<IContentSource>> ProviderSources = ContentSourceProvider->GetContentSources();
-		for (auto& EachSourceProvider : ProviderSources)
+		for (const TSharedRef<IContentSource>& EachSourceProvider : ProviderSources)
 		{
-			FFeaturePackContentSource* Source = (FFeaturePackContentSource*)&EachSourceProvider.Get();
-			FString ID = EachSourceProvider->GetIdent();
-			if( ID.IsEmpty() == false )
+			FFeaturePackContentSource* Source = (FFeaturePackContentSource*) &EachSourceProvider.Get();
+			FStringView ID = EachSourceProvider->GetIdent();
+			if (!ID.IsEmpty())
 			{
-				PackMap.Add(ID, Source);
+				PackMap.Add(FString(ID), Source);
 			}
 		}
 	}
@@ -690,7 +690,7 @@ bool FFeaturePackContentSource::ParseManifestString(const FString& ManifestStrin
 	if (ManifestReader->GetErrorMessage().IsEmpty() == false)
 	{
 		RecordAndLogError(FString::Printf(TEXT("Error in Feature pack %s. Failed to parse manifest: %s"), *FeaturePackPath, *ManifestReader->GetErrorMessage()));
-		Category = EContentSourceCategory::Unknown;
+		Categories = { EContentSourceCategory::Unknown };
 		return false;
 	}
 
@@ -707,7 +707,7 @@ bool FFeaturePackContentSource::ParseManifestString(const FString& ManifestStrin
 	if (TryValidateManifestObject(ManifestObject, ManifestObjectErrorMessage) == false)
 	{
 		RecordAndLogError(FString::Printf(TEXT("Error in Feature pack %s. Manifest object error: %s"), *FeaturePackPath, **ManifestObjectErrorMessage));
-		Category = EContentSourceCategory::Unknown;
+		Categories = { EContentSourceCategory::Unknown };
 		return false;
 	}
 
@@ -761,10 +761,25 @@ bool FFeaturePackContentSource::ParseManifestString(const FString& ManifestStrin
 	SortKey = FeaturePackPath;
 	ManifestObject->TryGetStringField("SortKey", SortKey);
 
-	FString CategoryString = ManifestObject->GetStringField("Category");
+	TArray<FString> CategoryStrings;
+	if (!ManifestObject->TryGetStringArrayField("Category", CategoryStrings))
+	{
+		FString CategoryString = ManifestObject->GetStringField("Category");
+		if (!CategoryString.IsEmpty())
+		{
+			CategoryStrings.Add(CategoryString);
+		}
+	}
+
 	UEnum* Enum = FindObjectChecked<UEnum>(ANY_PACKAGE, TEXT("EContentSourceCategory"));
-	int32 EnumValue = Enum->GetValueByName(FName(*CategoryString));
-	Category = EnumValue != INDEX_NONE ? (EContentSourceCategory)EnumValue : EContentSourceCategory::Unknown;
+	for (const FString& CategoryString : CategoryStrings)
+	{
+		int32 EnumValue = Enum->GetValueByName(FName(*CategoryString));
+		if (EnumValue != INDEX_NONE)
+		{
+			Categories.Add((EContentSourceCategory) EnumValue);
+		}
+	}
 
 	// Thumbnail filename
 	IconFilename = ManifestObject->GetStringField("Thumbnail");
diff --git a/Engine/Source/Editor/AddContentDialog/Private/ContentSourceProviders/FeaturePack/FeaturePackContentSourceProvider.cpp b/Engine/Source/Editor/AddContentDialog/Private/ContentSourceProviders/FeaturePack/FeaturePackContentSourceProvider.cpp
--- a/Engine/Source/Editor/AddContentDialog/Private/ContentSourceProviders/FeaturePack/FeaturePackContentSourceProvider.cpp
+++ b/Engine/Source/Editor/AddContentDialog/Private/ContentSourceProviders/FeaturePack/FeaturePackContentSourceProvider.cpp
@@ -40,7 +40,7 @@ FFeaturePackContentSourceProvider::FFeaturePackContentSourceProvider()
 	RefreshFeaturePacks();
 }
 
-const TArray<TSharedRef<IContentSource>> FFeaturePackContentSourceProvider::GetContentSources()
+const TArray<TSharedRef<IContentSource>>& FFeaturePackContentSourceProvider::GetContentSources() const
 {
 	return ContentSources;
 }
diff --git a/Engine/Source/Editor/AddContentDialog/Private/ContentSourceProviders/FeaturePack/FeaturePackContentSourceProvider.h b/Engine/Source/Editor/AddContentDialog/Private/ContentSourceProviders/FeaturePack/FeaturePackContentSourceProvider.h
--- a/Engine/Source/Editor/AddContentDialog/Private/ContentSourceProviders/FeaturePack/FeaturePackContentSourceProvider.h
+++ b/Engine/Source/Editor/AddContentDialog/Private/ContentSourceProviders/FeaturePack/FeaturePackContentSourceProvider.h
@@ -12,7 +12,7 @@ class FFeaturePackContentSourceProvider : public IContentSourceProvider
 {
 public:
 	FFeaturePackContentSourceProvider();
-	virtual const TArray<TSharedRef<IContentSource>> GetContentSources() override;
+	virtual const TArray<TSharedRef<IContentSource>>& GetContentSources() const override;
 	virtual void SetContentSourcesChanged(FOnContentSourcesChanged OnContentSourcesChangedIn) override;
 
 	virtual ~FFeaturePackContentSourceProvider();
diff --git a/Engine/Source/Editor/AddContentDialog/Private/IContentSource.h b/Engine/Source/Editor/AddContentDialog/Private/IContentSource.h
--- a/Engine/Source/Editor/AddContentDialog/Private/IContentSource.h
+++ b/Engine/Source/Editor/AddContentDialog/Private/IContentSource.h
@@ -40,7 +40,7 @@ public:
 	}
 
 	/** Gets the raw binary image data in PNG format. */
-	const TSharedPtr<TArray<uint8>> GetData() const
+	const TSharedPtr<TArray<uint8>>& GetData() const
 	{
 		return Data;
 	}
@@ -68,7 +68,7 @@ public:
 	}
 
 	/** Gets the iso 2-letter language specifier for this text. */
-	FString GetTwoLetterLanguage() const
+	const FString& GetTwoLetterLanguage() const
 	{
 		return TwoLetterLanguage;
 	}
@@ -88,29 +88,31 @@ private:
 class IContentSource
 {
 public:
+	virtual ~IContentSource() = default;
+
 	/** Gets the name of the content source as an array of localized strings. */
-	virtual TArray<FLocalizedText> GetLocalizedNames() const = 0;
+	virtual const TArray<FLocalizedText>& GetLocalizedNames() const = 0;
 
 	/** Gets the description of the content source as an array or localized strings. */
-	virtual TArray<FLocalizedText> GetLocalizedDescriptions() const = 0;
+	virtual const TArray<FLocalizedText>& GetLocalizedDescriptions() const = 0;
 
 	/** Gets the category for the content source. */
-	virtual EContentSourceCategory GetCategory() const = 0;
+	virtual const TArray<EContentSourceCategory>& GetCategories() const = 0;
 
 	/** Gets the image data for the icon which should represent the content source in the UI. */
 	virtual TSharedPtr<FImageData> GetIconData() const = 0;
 
 	/** Gets an array or image data for screenshots for the content source. */
-	virtual TArray<TSharedPtr<FImageData>> GetScreenshotData() const = 0;
+	virtual const TArray<TSharedPtr<FImageData>>& GetScreenshotData() const = 0;
 
 	/** Gets the asset types used in this pack. */
-	virtual TArray<FLocalizedText> GetLocalizedAssetTypes() const = 0;
+	virtual const TArray<FLocalizedText>& GetLocalizedAssetTypes() const = 0;
 
 	/** Gets the class types used in this pack. */
-	virtual FString GetClassTypesUsed() const = 0;
+	virtual const FString& GetClassTypesUsed() const = 0;
 
 	/** Gets the category for the content source. */
-	virtual FString GetSortKey() const = 0;
+	virtual const FString& GetSortKey() const = 0;
 	/*
 	 * Installs the content in the content source to the specific path. 
 	 * @returns true if install succeeded
@@ -120,8 +122,6 @@ public:
 	/** Is the data in this content valid. */
 	virtual bool IsDataValid() const = 0;
 
-	virtual ~IContentSource() { };
-
 	/** Gets the identity of the content. */
-	virtual FString GetIdent() const = 0;
+	virtual const FString& GetIdent() const = 0;
 };
diff --git a/Engine/Source/Editor/AddContentDialog/Private/IContentSourceProvider.h b/Engine/Source/Editor/AddContentDialog/Private/IContentSourceProvider.h
--- a/Engine/Source/Editor/AddContentDialog/Private/IContentSourceProvider.h
+++ b/Engine/Source/Editor/AddContentDialog/Private/IContentSourceProvider.h
@@ -13,7 +13,7 @@ public:
 
 public:
 	/** Gets the available content sources. */
-	virtual const TArray<TSharedRef<IContentSource>> GetContentSources() = 0;
+	virtual const TArray<TSharedRef<IContentSource>>& GetContentSources() const = 0;
 
 	/** Sets the delegate which will be executed when the avaialble content sources change */
 	virtual void SetContentSourcesChanged(FOnContentSourcesChanged OnContentSourcesChangedIn) = 0;
diff --git a/Engine/Source/Editor/AddContentDialog/Private/SAddContentWidget.cpp b/Engine/Source/Editor/AddContentDialog/Private/SAddContentWidget.cpp
--- a/Engine/Source/Editor/AddContentDialog/Private/SAddContentWidget.cpp
+++ b/Engine/Source/Editor/AddContentDialog/Private/SAddContentWidget.cpp
@@ -358,7 +358,7 @@ TSharedRef<SWidget> SAddContentWidget::CreateContentSourceDetail(TSharedPtr<FCon
 			.Padding(FMargin(10, 0, 0, 5))
 			[
 				SNew(STextBlock)
-				.Text(FText::FromString(ContentSource->GetClassTypes()))
+				.Text(FText::FromStringView(ContentSource->GetClassTypes()))
 				.Visibility(ContentSource->GetClassTypes().IsEmpty() == false ? EVisibility::Visible : EVisibility::Collapsed)
 				.AutoWrapText(true)
 			]
@@ -373,7 +373,7 @@ TSharedRef<SWidget> SAddContentWidget::CreateScreenshotCarousel(TSharedPtr<FCont
 		.NavigationBarStyle(FWidgetCarouselModuleStyle::Get(), "CarouselNavigationBar")
 		.NavigationButtonStyle(FWidgetCarouselModuleStyle::Get(), "CarouselNavigationButton")
 		.OnGenerateWidget(this, &SAddContentWidget::CreateScreenshotWidget)
-		.WidgetItemsSource(ContentSource->GetScreenshotBrushes());
+		.WidgetItemsSource(&ContentSource->GetScreenshotBrushes());
 }
 
 void SAddContentWidget::SearchTextChanged(const FText& SearchText)
diff --git a/Engine/Source/Editor/AddContentDialog/Private/ViewModels/AddContentWidgetViewModel.cpp b/Engine/Source/Editor/AddContentDialog/Private/ViewModels/AddContentWidgetViewModel.cpp
--- a/Engine/Source/Editor/AddContentDialog/Private/ViewModels/AddContentWidgetViewModel.cpp
+++ b/Engine/Source/Editor/AddContentDialog/Private/ViewModels/AddContentWidgetViewModel.cpp
@@ -21,17 +21,18 @@ TSharedRef<FAddContentWidgetViewModel> FAddContentWidgetViewModel::CreateShared(
 
 void FAddContentWidgetViewModel::Initialize()
 {
+	ContentSourceFilter = MakeShared<ContentSourceTextFilter>(
+		ContentSourceTextFilter::FItemToStringArray::CreateSP(this, &FAddContentWidgetViewModel::TransformContentSourceToStrings)
+	);
+
 	IAddContentDialogModule& AddContentDialogModule = FModuleManager::LoadModuleChecked<IAddContentDialogModule>("AddContentDialog");
 
-	for (auto ContentSourceProvider : *AddContentDialogModule.GetContentSourceProviderManager()->GetContentSourceProviders())
+	for (const TSharedRef<IContentSourceProvider>& ContentSourceProvider : *AddContentDialogModule.GetContentSourceProviderManager()->GetContentSourceProviders())
 	{
 		ContentSourceProviders.Add(ContentSourceProvider);
 		ContentSourceProvider->SetContentSourcesChanged(FOnContentSourcesChanged::CreateSP(this, &FAddContentWidgetViewModel::ContentSourcesChanged));
 	}
 
-	ContentSourceFilter = TSharedPtr<ContentSourceTextFilter>(new ContentSourceTextFilter(
-		ContentSourceTextFilter::FItemToStringArray::CreateSP(this, &FAddContentWidgetViewModel::TransformContentSourceToStrings)));
-
 	BuildContentSourceViewModels();
 }
 
@@ -112,48 +113,59 @@ void FAddContentWidgetViewModel::BuildContentSourceViewModels()
 	FilteredContentSourceViewModels.Empty();
 	CategoryToSelectedContentSourceMap.Empty();
 	
-	// List of categorys we dont want to see
+	// List of categories we don't want to see
 	TArray<EContentSourceCategory> FilteredCategories;
 	FilteredCategories.Add(EContentSourceCategory::SharedPack);
 	FilteredCategories.Add(EContentSourceCategory::Unknown);
 
-	for (auto ContentSourceProvider : ContentSourceProviders)
+	TSet<EContentSourceCategory> FoundCategories;
+
+	for (const TSharedPtr<IContentSourceProvider>& ContentSourceProvider : ContentSourceProviders)
 	{
-		for (auto ContentSource : ContentSourceProvider->GetContentSources())
+		for (const TSharedRef<IContentSource>& ContentSource : ContentSourceProvider->GetContentSources())
 		{
-			// Check if we want to see this category
-			if( FilteredCategories.Contains( ContentSource->GetCategory()) == false )
+			// Check if we want to see this content source - true unless all its categories are filtered out
+			bool bAnyVisible = false;
+			for (EContentSourceCategory ContentCategory : ContentSource->GetCategories())
 			{
-				TSharedPtr<FContentSourceViewModel> ContentSourceViewModel = MakeShareable(new FContentSourceViewModel(ContentSource));
-				if (Categories.Contains(ContentSourceViewModel->GetCategory()) == false)
+				if (!FilteredCategories.Contains(ContentCategory))
 				{
-					Categories.Add(ContentSourceViewModel->GetCategory());
+					FoundCategories.Add(ContentCategory);
+					bAnyVisible = true;
 				}
-				ContentSourceViewModels.Add(ContentSourceViewModel);
+			}
+
+			if (bAnyVisible)
+			{
+				ContentSourceViewModels.Add(MakeShared<FContentSourceViewModel>(ContentSource));
 			}
 		}
 	}
 
-	if (Categories.Num() > 0)
+	for (EContentSourceCategory Found : FoundCategories)
 	{
-		Categories.Sort();
-		// Update the current selection for all categories.  Do this in reverse order so that the first category
-		// remains selected when finished.
-		for (int i = Categories.Num() - 1; i >= 0; i--)
-		{
-			SelectedCategory = Categories[i];
-			UpdateFilteredContentSourcesAndSelection(false);
-		}
+		Categories.Add(FCategoryViewModel(Found));
+	}
+
+	Categories.Sort();
+
+	// Update the current selection for all categories.  Do this in reverse order so that the first category
+	// remains selected when finished.
+	for (int i = Categories.Num() - 1; i >= 0; i--)
+	{
+		SelectedCategory = Categories[i];
+		UpdateFilteredContentSourcesAndSelection(false);
 	}
+
 	OnCategoriesChanged.ExecuteIfBound();
 }
 
 void FAddContentWidgetViewModel::UpdateFilteredContentSourcesAndSelection(bool bAllowEmptySelection)
 {
 	FilteredContentSourceViewModels.Empty();
-	for (auto ContentSource : ContentSourceViewModels)
+	for (const TSharedPtr<FContentSourceViewModel>& ContentSource : ContentSourceViewModels)
 	{
-		if ((ContentSource->GetCategory() == SelectedCategory) && (ContentSourceFilter.IsValid() == false || ContentSourceFilter->PassesFilter(ContentSource)))
+		if (ContentSource->GetCategories().Contains(SelectedCategory) && ContentSourceFilter->PassesFilter(ContentSource))
 		{
 			FilteredContentSourceViewModels.Add(ContentSource);
 		}
diff --git a/Engine/Source/Editor/AddContentDialog/Private/ViewModels/CategoryViewModel.cpp b/Engine/Source/Editor/AddContentDialog/Private/ViewModels/CategoryViewModel.cpp
--- a/Engine/Source/Editor/AddContentDialog/Private/ViewModels/CategoryViewModel.cpp
+++ b/Engine/Source/Editor/AddContentDialog/Private/ViewModels/CategoryViewModel.cpp
@@ -17,16 +17,6 @@ FCategoryViewModel::FCategoryViewModel(EContentSourceCategory InCategory)
 	Initialize();
 }
 
-FText FCategoryViewModel::GetText() const
-{
-	return Text;
-}
-
-uint32 FCategoryViewModel::GetTypeHash() const
-{
-	return (uint32)Category;
-}
-
 void FCategoryViewModel::Initialize()
 {
 	switch (Category)
diff --git a/Engine/Source/Editor/AddContentDialog/Private/ViewModels/CategoryViewModel.h b/Engine/Source/Editor/AddContentDialog/Private/ViewModels/CategoryViewModel.h
--- a/Engine/Source/Editor/AddContentDialog/Private/ViewModels/CategoryViewModel.h
+++ b/Engine/Source/Editor/AddContentDialog/Private/ViewModels/CategoryViewModel.h
@@ -16,7 +16,7 @@ public:
 	FCategoryViewModel(EContentSourceCategory InCategory);
 
 	/** Gets the display name of the category. */
-	FText GetText() const;
+	const FText& GetText() const { return Text; }
 
 	inline bool operator==(const FCategoryViewModel& Other) const
 	{
@@ -28,7 +28,7 @@ public:
 		return SortID < Other.SortID;
 	}
 
-	uint32 GetTypeHash() const;
+	uint32 GetTypeHash() const { return (uint32) Category; }
 
 private:
 	void Initialize();
diff --git a/Engine/Source/Editor/AddContentDialog/Private/ViewModels/ContentSourceViewModel.cpp b/Engine/Source/Editor/AddContentDialog/Private/ViewModels/ContentSourceViewModel.cpp
--- a/Engine/Source/Editor/AddContentDialog/Private/ViewModels/ContentSourceViewModel.cpp
+++ b/Engine/Source/Editor/AddContentDialog/Private/ViewModels/ContentSourceViewModel.cpp
@@ -17,86 +17,102 @@ FContentSourceViewModel::FContentSourceViewModel(TSharedPtr<IContentSource> Cont
 {
 	ContentSource = ContentSourceIn;
 	SetupBrushes();
-	Category = FCategoryViewModel(ContentSource->GetCategory());
+
+	for (EContentSourceCategory Category : ContentSource->GetCategories())
+	{
+		Categories.Add(FCategoryViewModel(Category));
+	}
 }
 
-TSharedPtr<IContentSource> FContentSourceViewModel::GetContentSource()
+const TSharedPtr<IContentSource>& FContentSourceViewModel::GetContentSource() const
 {
 	return ContentSource;
 }
 
-FText FContentSourceViewModel::GetName()
+const FText& FContentSourceViewModel::GetName() const
 {
-	const FString CurrentLanguage = FInternationalization::Get().GetCurrentLanguage()->GetName();
+	const FString& CurrentLanguage = FInternationalization::Get().GetCurrentLanguage()->GetName();
 	if (CachedNameText.Language != CurrentLanguage)
 	{
-		CachedNameText.Language = CurrentLanguage;
-		CachedNameText.Text = ChooseLocalizedText(ContentSource->GetLocalizedNames(), CurrentLanguage);
+		CachedNameText =
+			FCachedContentText
+		{
+			CurrentLanguage,
+			ChooseLocalizedText(ContentSource->GetLocalizedNames(), CurrentLanguage)
+		};
 	}
 	return CachedNameText.Text;
 }
 
-FText FContentSourceViewModel::GetDescription()
+const FText& FContentSourceViewModel::GetDescription() const
 {
-	const FString CurrentLanguage = FInternationalization::Get().GetCurrentLanguage()->GetName();
+	const FString& CurrentLanguage = FInternationalization::Get().GetCurrentLanguage()->GetName();
 	if (CachedDescriptionText.Language != CurrentLanguage)
 	{
-		CachedDescriptionText.Language = CurrentLanguage;
-		CachedDescriptionText.Text = ChooseLocalizedText(ContentSource->GetLocalizedDescriptions(), CurrentLanguage);
+		CachedDescriptionText = 
+			FCachedContentText 
+			{ 
+				CurrentLanguage, 
+				ChooseLocalizedText(ContentSource->GetLocalizedDescriptions(), CurrentLanguage) 
+			};
 	}
 	return CachedDescriptionText.Text;
 }
 
-FText FContentSourceViewModel::GetAssetTypes()
+const FText& FContentSourceViewModel::GetAssetTypes() const
 {
-	const FString CurrentLanguage = FInternationalization::Get().GetCurrentLanguage()->GetName();
+	const FString& CurrentLanguage = FInternationalization::Get().GetCurrentLanguage()->GetName();
 	if (CachedAssetTypeText.Language != CurrentLanguage)
 	{
-		CachedAssetTypeText.Language = CurrentLanguage;
-		CachedAssetTypeText.Text = ChooseLocalizedText(ContentSource->GetLocalizedAssetTypes(), CurrentLanguage);
+		CachedAssetTypeText =
+			FCachedContentText
+		{
+			CurrentLanguage,
+			ChooseLocalizedText(ContentSource->GetLocalizedAssetTypes(), CurrentLanguage)
+		};
 	}
 	return CachedAssetTypeText.Text;
 }
 
-FString FContentSourceViewModel::GetClassTypes()
+FStringView FContentSourceViewModel::GetClassTypes() const
 {
 	return ContentSource->GetClassTypesUsed();
 }
 
-FCategoryViewModel FContentSourceViewModel::GetCategory()
+const TArray<FCategoryViewModel>& FContentSourceViewModel::GetCategories() const
 {
-	return Category;
+	return Categories;
 }
 
-TSharedPtr<FSlateBrush> FContentSourceViewModel::GetIconBrush()
+const TSharedPtr<FSlateBrush>& FContentSourceViewModel::GetIconBrush() const
 {
 	return IconBrush;
 }
 
-TArray<TSharedPtr<FSlateBrush>>* FContentSourceViewModel::GetScreenshotBrushes()
+const TArray<TSharedPtr<FSlateBrush>>& FContentSourceViewModel::GetScreenshotBrushes() const
 {
-	return &ScreenshotBrushes;
+	return ScreenshotBrushes;
 }
 
 void FContentSourceViewModel::SetupBrushes()
 {
 	if (ContentSource->GetIconData().IsValid())
 	{
-	    FString IconBrushName = GetName().ToString() + "_" + ContentSource->GetIconData()->GetName();
-	    IconBrush = CreateBrushFromRawData(IconBrushName, *ContentSource->GetIconData()->GetData());
+		const FString IconBrushName = GetName().ToString() + "_" + ContentSource->GetIconData()->GetName();
+		IconBrush = CreateBrushFromRawData(IconBrushName, *ContentSource->GetIconData()->GetData());
 	}
 
-	for (TSharedPtr<FImageData> ScreenshotData : ContentSource->GetScreenshotData())
+	for (const TSharedPtr<FImageData>& ScreenshotData : ContentSource->GetScreenshotData())
 	{
 		if (ScreenshotData.IsValid() == true)
 		{
-		    FString ScreenshotBrushName = GetName().ToString() + "_" + ScreenshotData->GetName();
-		    ScreenshotBrushes.Add(CreateBrushFromRawData(ScreenshotBrushName, *ScreenshotData->GetData()));
+			const FString ScreenshotBrushName = GetName().ToString() + "_" + ScreenshotData->GetName();
+			ScreenshotBrushes.Add(CreateBrushFromRawData(ScreenshotBrushName, *ScreenshotData->GetData()));
 	    }
     }
 }
 
-TSharedPtr<FSlateDynamicImageBrush> FContentSourceViewModel::CreateBrushFromRawData(FString ResourceNamePrefix, const TArray< uint8 >& RawData) const
+TSharedPtr<FSlateDynamicImageBrush> FContentSourceViewModel::CreateBrushFromRawData(const FString& ResourceNamePrefix, const TArray<uint8>& RawData) const
 {
 	TSharedPtr< FSlateDynamicImageBrush > Brush;
 
@@ -128,7 +144,7 @@ TSharedPtr<FSlateDynamicImageBrush> FContentSourceViewModel::CreateBrushFromRawD
 	return Brush;
 }
 
-FText FContentSourceViewModel::ChooseLocalizedText(const TArray<FLocalizedText>& Choices, const FString& InCurrentLanguage)
+FText FContentSourceViewModel::ChooseLocalizedText(const TArray<FLocalizedText>& Choices, const FString& InCurrentLanguage) const
 {
 	auto FindLocalizedTextForCulture = [&Choices](const FString& InCulture)
 	{
diff --git a/Engine/Source/Editor/AddContentDialog/Private/ViewModels/ContentSourceViewModel.h b/Engine/Source/Editor/AddContentDialog/Private/ViewModels/ContentSourceViewModel.h
--- a/Engine/Source/Editor/AddContentDialog/Private/ViewModels/ContentSourceViewModel.h
+++ b/Engine/Source/Editor/AddContentDialog/Private/ViewModels/ContentSourceViewModel.h
@@ -15,38 +15,38 @@ public:
 	FContentSourceViewModel(TSharedPtr<IContentSource> ContentSourceIn);
 
 	/** Gets the content source represented by this view model. */
-	TSharedPtr<IContentSource> GetContentSource();
+	const TSharedPtr<IContentSource>& GetContentSource() const;
 
 	/** Gets the display name for this content source. */
-	FText GetName();
+	const FText& GetName() const;
 
 	/** Gets the description of this content source. */
-	FText GetDescription();
+	const FText& GetDescription() const;
 
 	/** Gets the asset types used in this content source. */
-	FText GetAssetTypes();
+	const FText& GetAssetTypes() const;
 
 	/** Gets the class types used in this content source. */
-	FString GetClassTypes();
+	FStringView GetClassTypes() const;
 
-	/** Gets the view model for the category for this content source. */
-	FCategoryViewModel GetCategory();
+	/** Gets the view models for the categories for this content source. */
+	const TArray<FCategoryViewModel>& GetCategories() const;
 
 	/** Gets the brush which should be used to draw the icon representation of this content source. */
-	TSharedPtr<FSlateBrush> GetIconBrush();
+	const TSharedPtr<FSlateBrush>& GetIconBrush() const;
 
 	/** Gets an array or brushes which should be used to display screenshots for this content source. */
-	TArray<TSharedPtr<FSlateBrush>>* GetScreenshotBrushes();
+	const TArray<TSharedPtr<FSlateBrush>>& GetScreenshotBrushes() const;
 
 private:
 	/** Sets up brushes from the images data supplied by the IContentSource. */
 	void SetupBrushes();
 
 	/** Creates a slate brush from raw binary PNG formatted image data and the supplied prefix. */
-	TSharedPtr<FSlateDynamicImageBrush> CreateBrushFromRawData(FString ResourceNamePrefix, const TArray<uint8>& RawData) const;
+	TSharedPtr<FSlateDynamicImageBrush> CreateBrushFromRawData(const FString& ResourceNamePrefix, const TArray<uint8>& RawData) const;
 
 	/** Selects the text from an array which matches the given language. */
-	FText ChooseLocalizedText(const TArray<FLocalizedText>& Choices, const FString& InCurrentLanguage);
+	FText ChooseLocalizedText(const TArray<FLocalizedText>& Choices, const FString& InCurrentLanguage) const;
 
 private:
 	struct FCachedContentText
@@ -64,17 +64,17 @@ private:
 	/** An array or brushes which should be used to display screenshots for this content source. */
 	TArray<TSharedPtr<FSlateBrush>> ScreenshotBrushes;
 
-	/** The view model for the category for this content source. */
-	FCategoryViewModel Category;
+	/** The view models for the categories for this content source. */
+	TArray<FCategoryViewModel> Categories;
 
 	/** The information used/returned the last time the name of the content source was requested. */
-	FCachedContentText CachedNameText;
+	mutable FCachedContentText CachedNameText;
 
 	/** The information used/returned the last time the description of the content source was requested. */
-	FCachedContentText CachedDescriptionText;
+	mutable FCachedContentText CachedDescriptionText;
 
 	/** The information used/returned the last time the asset types of the content source was requested. */
-	FCachedContentText CachedAssetTypeText;
+	mutable FCachedContentText CachedAssetTypeText;
 
 	/** Keeps track of a unique increasing id which is appended to each brush name.  This avoids an issue
 		where two brushes are created with the same name, and then both brushes texture data gets deleted
diff --git a/Engine/Source/Editor/AddContentDialog/Public/FeaturePackContentSource.h b/Engine/Source/Editor/AddContentDialog/Public/FeaturePackContentSource.h
--- a/Engine/Source/Editor/AddContentDialog/Public/FeaturePackContentSource.h
+++ b/Engine/Source/Editor/AddContentDialog/Public/FeaturePackContentSource.h
@@ -42,13 +42,13 @@ public:
 	}
 
 	/** Gets the iso 2-letter language specifier for this text. */
-	FString GetTwoLetterLanguage() const
+	const FString& GetTwoLetterLanguage() const
 	{
 		return TwoLetterLanguage;
 	}
 
 	/** Gets the array of tags in the language specified. */
-	TArray<FText> GetTags() const
+	const TArray<FText>& GetTags() const
 	{
 		return Tags;
 	}
@@ -141,17 +141,17 @@ public:
 
 	virtual ~FFeaturePackContentSource();
 
-	virtual TArray<FLocalizedText> GetLocalizedNames() const override;
-	virtual TArray<FLocalizedText> GetLocalizedDescriptions() const override;
-	
-	virtual EContentSourceCategory GetCategory() const override;
-	virtual TArray<FLocalizedText> GetLocalizedAssetTypes() const override;
-	virtual FString GetSortKey() const override;
-	virtual FString GetClassTypesUsed() const override;
+	virtual const TArray<FLocalizedText>& GetLocalizedNames() const override;
+	virtual const TArray<FLocalizedText>& GetLocalizedDescriptions() const override;
 	
+	virtual const TArray<EContentSourceCategory>& GetCategories() const override;
+	virtual const TArray<FLocalizedText>& GetLocalizedAssetTypes() const override;
+	virtual const FString& GetSortKey() const override;
+	virtual const FString& GetClassTypesUsed() const override;
+	virtual const FString& GetIdent() const override;
 	virtual TSharedPtr<FImageData> GetIconData() const override;
-	virtual TArray<TSharedPtr<FImageData>> GetScreenshotData() const override;
-	FString GetFocusAssetName() const;
+	virtual const TArray<TSharedPtr<FImageData>>& GetScreenshotData() const override;
+	const FString& GetFocusAssetName() const;
 
 	virtual bool InstallToProject(FString InstallPath) override;
 
@@ -186,9 +186,6 @@ public:
 	
 	void BuildListOfAdditionalFiles(TArray<FString>& AdditionalFileSourceList,TArray<FString>& FileList, bool& bContainsSourceFiles);
 
-	/* Get the identifier of the pack. */
-	virtual FString GetIdent() const;
-
 private:
 	static void ParseAndImportPacks();
 	bool LoadPakFileToBuffer(FPakPlatformFile& PakPlatformFile, FString Path, TArray<uint8>& Buffer);
@@ -227,8 +224,8 @@ private:
 	/* Array of localised descriptions */
 	TArray<FLocalizedText> LocalizedDescriptions;
 	
-	/* Defines the type of feature pack this is */
-	EContentSourceCategory Category;
+	/* Defines the type(s) of feature pack this is */
+	TArray<EContentSourceCategory> Categories;
 	
 	/* Filename of the icon */
 	FString IconFilename;
