Screenshot carousel in the "Add Content or Feature Pack..." dialog removed for the time being and replaced with a plain SImage due to currently rendering black because of a presumed bug in SFxWidget.

There are currently no intentions to add multiple screenshots, so the previous SCarouselWidget was unnecessary anyway, but the decision can be revisited in the future if desired.

diff --git a/Engine/Source/Editor/AddContentDialog/Private/SAddContentWidget.cpp b/Engine/Source/Editor/AddContentDialog/Private/SAddContentWidget.cpp
--- a/Engine/Source/Editor/AddContentDialog/Private/SAddContentWidget.cpp
+++ b/Engine/Source/Editor/AddContentDialog/Private/SAddContentWidget.cpp
@@ -369,11 +369,12 @@ TSharedRef<SWidget> SAddContentWidget::CreateContentSourceDetail(TSharedPtr<FCon
 
 TSharedRef<SWidget> SAddContentWidget::CreateScreenshotCarousel(TSharedPtr<FContentSourceViewModel> ContentSource)
 {
-	return SNew(SWidgetCarouselWithNavigation<TSharedPtr<FSlateBrush>>)
-		.NavigationBarStyle(FWidgetCarouselModuleStyle::Get(), "CarouselNavigationBar")
-		.NavigationButtonStyle(FWidgetCarouselModuleStyle::Get(), "CarouselNavigationButton")
-		.OnGenerateWidget(this, &SAddContentWidget::CreateScreenshotWidget)
-		.WidgetItemsSource(&ContentSource->GetScreenshotBrushes());
+	TSharedPtr<FSlateBrush> Brush;
+	if (ContentSource->GetScreenshotBrushes().Num() > 0)
+	{
+		Brush = ContentSource->GetScreenshotBrushes()[0];
+	}
+	return CreateScreenshotWidget(Brush);
 }
 
 void SAddContentWidget::SearchTextChanged(const FText& SearchText)
