[Backout] CL 19230249
Changelist 19230249 swapped out the SWidgetCarouselWithNavigation widget in the Add Feature or Content Pack window in Release-5.0 to avoid a rendering bug in SFxWidget. That bug has been fixed in Main, so we can re-enable the SWidgetCarouselWithNavigation.

diff --git a/Engine/Source/Editor/AddContentDialog/Private/SAddContentWidget.cpp b/Engine/Source/Editor/AddContentDialog/Private/SAddContentWidget.cpp
--- a/Engine/Source/Editor/AddContentDialog/Private/SAddContentWidget.cpp
+++ b/Engine/Source/Editor/AddContentDialog/Private/SAddContentWidget.cpp
@@ -369,12 +369,11 @@ TSharedRef<SWidget> SAddContentWidget::CreateContentSourceDetail(TSharedPtr<FCon
 
 TSharedRef<SWidget> SAddContentWidget::CreateScreenshotCarousel(TSharedPtr<FContentSourceViewModel> ContentSource)
 {
-	TSharedPtr<FSlateBrush> Brush;
-	if (ContentSource->GetScreenshotBrushes().Num() > 0)
-	{
-		Brush = ContentSource->GetScreenshotBrushes()[0];
-	}
-	return CreateScreenshotWidget(Brush);
+	return SNew(SWidgetCarouselWithNavigation<TSharedPtr<FSlateBrush>>)
+		.NavigationBarStyle(FWidgetCarouselModuleStyle::Get(), "CarouselNavigationBar")
+		.NavigationButtonStyle(FWidgetCarouselModuleStyle::Get(), "CarouselNavigationButton")
+		.OnGenerateWidget(this, &SAddContentWidget::CreateScreenshotWidget)
+		.WidgetItemsSource(&ContentSource->GetScreenshotBrushes());
 }
 
 void SAddContentWidget::SearchTextChanged(const FText& SearchText)
