Fix rare crash when populating a new project with starter content.

- When populating the starter content, IAssetTools::ImportAssets will load and return one of the packages containing an asset and FFeaturePackContentSource will attempt to save that asset. The load can kick off async compilation which can still be in progress when the save occurs and the save will cause the package to become detached from its file on disk.
- This introduces the classic thread safety issue between editor bulkdata and package detachment.
- For now we fix this by flushing async compilation before doing operations that we know can cause this problem.
- This fix can be removed once UE-140944 is done.

diff --git a/Engine/Source/Editor/AddContentDialog/AddContentDialog.Build.cs b/Engine/Source/Editor/AddContentDialog/AddContentDialog.Build.cs
--- a/Engine/Source/Editor/AddContentDialog/AddContentDialog.Build.cs
+++ b/Engine/Source/Editor/AddContentDialog/AddContentDialog.Build.cs
@@ -15,21 +15,22 @@ public class AddContentDialog : ModuleRules
 
 		PrivateDependencyModuleNames.AddRange(
 			new string[] {
-				"Core",
 				"ApplicationCore",
-				"Slate",
-				"SlateCore",
+				"Core",
+				"CoreUObject",
+				"DesktopPlatform",
+				"DirectoryWatcher",
+				"EditorFramework",
+				"Engine",
+				"ImageWrapper",
 				"InputCore",
 				"Json",
-				"DirectoryWatcher",
-				"DesktopPlatform",
 				"PakFile",
-				"ImageWrapper",
-				"EditorFramework",
-				"UnrealEd",
-				"CoreUObject",				
-				"WidgetCarousel",
+				"Slate",
+				"SlateCore",
 				"ToolWidgets",
+				"UnrealEd",
+				"WidgetCarousel"
 			}
 		);
 	}
diff --git a/Engine/Source/Editor/AddContentDialog/Private/ContentSourceProviders/FeaturePack/FeaturePackContentSource.cpp b/Engine/Source/Editor/AddContentDialog/Private/ContentSourceProviders/FeaturePack/FeaturePackContentSource.cpp
--- a/Engine/Source/Editor/AddContentDialog/Private/ContentSourceProviders/FeaturePack/FeaturePackContentSource.cpp
+++ b/Engine/Source/Editor/AddContentDialog/Private/ContentSourceProviders/FeaturePack/FeaturePackContentSource.cpp
@@ -2,6 +2,7 @@
 
 #include "FeaturePackContentSource.h"
 
+#include "AssetCompilingManager.h"
 #include "AssetToolsModule.h"
 #include "Containers/Map.h"
 #include "Containers/StringFwd.h"
@@ -566,6 +567,10 @@ void FFeaturePackContentSource::ParseAndImportPacks()
 				{
 					ToSave.AddUnique(ImportedObject->GetOutermost());
 				}
+
+				// Make sure any async compilation kicked off during ImportAssets is completed before we save.
+				FAssetCompilingManager::Get().FinishAllCompilation();
+
 				FEditorFileUtils::PromptForCheckoutAndSave(ToSave, /*bCheckDirty=*/ false, /*bPromptToSave=*/ false);
 				PacksInserted++;
 			}
