Look for existing mounted pak files before attempting to make a new one.

diff --git a/Engine/Source/Editor/AddContentDialog/Private/ContentSourceProviders/FeaturePack/FeaturePackContentSource.cpp b/Engine/Source/Editor/AddContentDialog/Private/ContentSourceProviders/FeaturePack/FeaturePackContentSource.cpp
--- a/Engine/Source/Editor/AddContentDialog/Private/ContentSourceProviders/FeaturePack/FeaturePackContentSource.cpp
+++ b/Engine/Source/Editor/AddContentDialog/Private/ContentSourceProviders/FeaturePack/FeaturePackContentSource.cpp
@@ -177,15 +177,20 @@ FFeaturePackContentSource::FFeaturePackContentSource(FString InFeaturePackPath)
 	{
 		bContentsInPakFile = true;
 		MountPoint = FPaths::GameFeatureRootPrefix();
-		// Create a pak platform file and mount the feature pack file.
-		FPakPlatformFile PakPlatformFile;
-		PakPlatformFile.Initialize(&FPlatformFileManager::Get().GetPlatformFile(), TEXT(""));
-		
-		PakPlatformFile.Mount(*InFeaturePackPath, 0, *MountPoint);
+
+		const FString PakFileName = TEXT("PakFile");
+		FPakPlatformFile* PakPlatformFile = static_cast<FPakPlatformFile*>(FPlatformFileManager::Get().FindPlatformFile(*PakFileName));
+		if (!PakPlatformFile)
+		{
+			// Create a pak platform file and mount the feature pack file.
+			PakPlatformFile = static_cast<FPakPlatformFile*>(FPlatformFileManager::Get().GetPlatformFile(*PakFileName));
+			PakPlatformFile->Initialize(&FPlatformFileManager::Get().GetPlatformFile(), TEXT(""));
+		}
+		PakPlatformFile->Mount(*InFeaturePackPath, 0, *MountPoint);
 
 		// Gets the manifest file as a JSon string
 		TArray<uint8> ManifestBuffer;
-		if (LoadPakFileToBuffer(PakPlatformFile, FPaths::Combine(*MountPoint, TEXT("manifest.json")), ManifestBuffer) == false)
+		if (LoadPakFileToBuffer(*PakPlatformFile, FPaths::Combine(*MountPoint, TEXT("manifest.json")), ManifestBuffer) == false)
 		{
 			RecordAndLogError(FString::Printf(TEXT("Error in Feature pack %s. Cannot find manifest."), *FeaturePackPath));
 			return;
@@ -195,7 +200,7 @@ FFeaturePackContentSource::FFeaturePackContentSource(FString InFeaturePackPath)
 
 		if (ParseManifestString(ManifestString) == true)
 		{
-			LoadFeaturePackImageDataFromPackFile(PakPlatformFile);			
+			LoadFeaturePackImageDataFromPackFile(*PakPlatformFile);			
 		}
 	}
 	else
