Unmount content when mounting to extract the images.
This allows to read the full pak file from the mount point and not have the old data which caused the bug

diff --git a/Engine/Source/Editor/AddContentDialog/Private/ContentSourceProviders/FeaturePack/FeaturePackContentSource.cpp b/Engine/Source/Editor/AddContentDialog/Private/ContentSourceProviders/FeaturePack/FeaturePackContentSource.cpp
--- a/Engine/Source/Editor/AddContentDialog/Private/ContentSourceProviders/FeaturePack/FeaturePackContentSource.cpp
+++ b/Engine/Source/Editor/AddContentDialog/Private/ContentSourceProviders/FeaturePack/FeaturePackContentSource.cpp
@@ -202,6 +202,7 @@ FFeaturePackContentSource::FFeaturePackContentSource(FString InFeaturePackPath)
 		{
 			LoadFeaturePackImageDataFromPackFile(*PakPlatformFile);			
 		}
+		PakPlatformFile->Unmount(*InFeaturePackPath);
 	}
 	else
 	{
