[FeaturePack] Add support for Variants using external actors

diff --git a/Engine/Source/Editor/AddContentDialog/Private/ContentSourceProviders/FeaturePack/FeaturePackContentSource.cpp b/Engine/Source/Editor/AddContentDialog/Private/ContentSourceProviders/FeaturePack/FeaturePackContentSource.cpp
--- a/Engine/Source/Editor/AddContentDialog/Private/ContentSourceProviders/FeaturePack/FeaturePackContentSource.cpp
+++ b/Engine/Source/Editor/AddContentDialog/Private/ContentSourceProviders/FeaturePack/FeaturePackContentSource.cpp
@@ -591,17 +591,37 @@ void FFeaturePackContentSource::RecordAndLogError(const FString& ErrorString)
 
 void FFeaturePackContentSource::CopyAdditionalFilesToFolder( const FString& DestinationFolder, TArray<FString>& FilesCopied, bool &bHasSourceFiles, FString InGameFolder )
 {
-	FString ContentIdent = TEXT("Content/");
+	// Array of pairs - possible folder identifier and bool value indicating whether Content folder should be appended to BaseFolderName
+	// For example, we want to move contents of __ExternalActors__/ folder to Content/__ExternalActors__/
+	TStaticArray<TPair<FString, bool>, 3> FolderIdents;
+	FolderIdents[0] = { TEXT("__ExternalActors__/"), true };
+	FolderIdents[1] = { TEXT("__ExternalObjects__/"), true };
+	FolderIdents[2] = { TEXT("Content/"), false };
+
 	TArray<FString> FilesToAdd;
 	GetAdditionalFilesForPack(FilesToAdd, bHasSourceFiles);
 	for (int32 iFile = 0; iFile < FilesToAdd.Num(); ++iFile)
 	{
 		FString EachFile = FilesToAdd[iFile];
-		int32 ContentIndex;
-		ContentIndex = EachFile.Find(*ContentIdent);
-		if (ContentIndex != INDEX_NONE)
+
+		int32 BaseFolderNameIndex = INDEX_NONE;
+		int32 BaseFolderNameLength = 0;
+		bool bAppendContentFolderToBaseFolderName = false;
+
+		for (const TPair<FString, bool>& FolderIdentPair : FolderIdents)
 		{
-			FString ContentFile = EachFile.RightChop(ContentIndex);
+			if (int32 FolderNameIndex = EachFile.Find(FolderIdentPair.Key); FolderNameIndex != INDEX_NONE)
+			{
+				BaseFolderNameIndex = FolderNameIndex;
+				BaseFolderNameLength = FolderIdentPair.Key.Len();
+				bAppendContentFolderToBaseFolderName = FolderIdentPair.Value;
+				break;
+			}
+		}
+
+		if (BaseFolderNameIndex != INDEX_NONE)
+		{
+			FString ContentFile = EachFile.RightChop(BaseFolderNameIndex);
 			FString GameFolder = InGameFolder;
 			if(( GameFolder.IsEmpty() == false) && ( GameFolder.StartsWith(TEXT("/")) == false ))
 			{
@@ -610,7 +630,12 @@ void FFeaturePackContentSource::CopyAdditionalFilesToFolder( const FString& Dest
 			FPaths::NormalizeFilename(ContentFile);
 			if( GameFolder.IsEmpty() == false)
 			{
-				ContentFile.InsertAt(ContentIdent.Len() - 1, GameFolder);
+				ContentFile.InsertAt(BaseFolderNameLength - 1, GameFolder);
+			}
+
+			if (bAppendContentFolderToBaseFolderName)
+			{
+				ContentFile = TEXT("Content/") + ContentFile;
 			}
 
 			FString FinalDestination = DestinationFolder / ContentFile;
