Add log messages if a feature pack cannot be found or installed correctly.

#rnx

diff --git a/Engine/Source/Editor/AddContentDialog/Private/ContentSourceProviders/FeaturePack/FeaturePackContentSource.cpp b/Engine/Source/Editor/AddContentDialog/Private/ContentSourceProviders/FeaturePack/FeaturePackContentSource.cpp
--- a/Engine/Source/Editor/AddContentDialog/Private/ContentSourceProviders/FeaturePack/FeaturePackContentSource.cpp
+++ b/Engine/Source/Editor/AddContentDialog/Private/ContentSourceProviders/FeaturePack/FeaturePackContentSource.cpp
@@ -723,6 +723,14 @@ bool FFeaturePackContentSource::InsertAdditionalResources(TArray<FFeaturePackLev
 					NewContentSource->CopyAdditionalFilesToFolder(InDestinationFolder, FilesCopied, bHasSourceFiles, InAdditionalFeaturePacks[iExtraPack].MountName);
 					PacksInserted++;
 				}
+				else
+				{
+					UE_LOG(LogFeaturePack, Warning, TEXT("Unable to install feature pack: '%s'"), *FullPath);
+				}
+			}
+			else
+			{
+				UE_LOG(LogFeaturePack, Warning, TEXT("Unable to find feature pack: '%s'"), *FullPath);
 			}
 		}
 		// Add to the list of files copied
