PR #13849: Changed refresh of content sources to only happen if needed to speed up editor launch

Sligthly modified to avoid removing the const interface
Add validation that refreshing the content always only happen from the game-thread since it is calling a delegate

#FailedValidations Preflight

diff --git a/Engine/Source/Editor/AddContentDialog/Private/ContentSourceProviders/FeaturePack/FeaturePackContentSourceProvider.cpp b/Engine/Source/Editor/AddContentDialog/Private/ContentSourceProviders/FeaturePack/FeaturePackContentSourceProvider.cpp
--- a/Engine/Source/Editor/AddContentDialog/Private/ContentSourceProviders/FeaturePack/FeaturePackContentSourceProvider.cpp
+++ b/Engine/Source/Editor/AddContentDialog/Private/ContentSourceProviders/FeaturePack/FeaturePackContentSourceProvider.cpp
@@ -43,11 +43,17 @@ FFeaturePackContentSourceProvider::FFeaturePackContentSourceProvider()
 	TemplatePath = FPaths::RootDir() + TEXT("Templates/");
 	EnterpriseTemplatePath = FPaths::EnterpriseDir() + TEXT("Templates/");
 	StartUpDirectoryWatcher();
-	RefreshFeaturePacks();
 }
 
 const TArray<TSharedRef<IContentSource>>& FFeaturePackContentSourceProvider::GetContentSources() const
 {
+	if (!bAreContentSourcesUpToDate)
+	{
+		// Using mutable to avoid removing the const interface since this is just a caching mechanism.
+		FFeaturePackContentSourceProvider* MutableThis = const_cast<FFeaturePackContentSourceProvider*>(this);
+		MutableThis->RefreshFeaturePacks();
+	}
+
 	return ContentSources;
 }
 
@@ -89,7 +95,7 @@ void FFeaturePackContentSourceProvider::ShutDownDirectoryWatcher()
 
 void FFeaturePackContentSourceProvider::OnFeaturePackDirectoryChanged( const TArray<FFileChangeData>& FileChanges )
 {
-	RefreshFeaturePacks();
+	bAreContentSourcesUpToDate = false;
 }
 
 /** Sorting function sort keys. */
@@ -100,6 +106,8 @@ struct FFeaturePackCompareSortKey
 
 void FFeaturePackContentSourceProvider::RefreshFeaturePacks()
 {
+	TRACE_CPUPROFILER_EVENT_SCOPE(FFeaturePackContentSourceProvider::RefreshFeaturePacks);
+	
 	ContentSources.Empty();
 	// TODO improve this and seperate handling of .upack and manifest (loose) parsing
 	// first the .upack files
@@ -139,8 +147,11 @@ void FFeaturePackContentSourceProvider::RefreshFeaturePacks()
 		}
 	}
 
-
 	ContentSources.Sort(FFeaturePackCompareSortKey());
+	bAreContentSourcesUpToDate = true;
+
+	// Assume for now that this delegate doesn't support being called from any thread.
+	check(IsInGameThread());
 	OnContentSourcesChanged.ExecuteIfBound();
 }
 
diff --git a/Engine/Source/Editor/AddContentDialog/Private/ContentSourceProviders/FeaturePack/FeaturePackContentSourceProvider.h b/Engine/Source/Editor/AddContentDialog/Private/ContentSourceProviders/FeaturePack/FeaturePackContentSourceProvider.h
--- a/Engine/Source/Editor/AddContentDialog/Private/ContentSourceProviders/FeaturePack/FeaturePackContentSourceProvider.h
+++ b/Engine/Source/Editor/AddContentDialog/Private/ContentSourceProviders/FeaturePack/FeaturePackContentSourceProvider.h
@@ -59,5 +59,8 @@ private:
 	/** An array of the available Enterprise content sources. */
 	TArray<TSharedRef<IContentSource>> EnterpiseContentSources;
 
+	/* Are the cached content sources up to date*/
+	mutable bool bAreContentSourcesUpToDate = false;
+
 	FDelegateHandle DirectoryChangedDelegateHandle;
 };
