[AddContentDialog]
- Prevent duplicates by handling reentrancy when an update is sent while we're already rebuilding the view model

#Synced-CL 52335721

diff --git a/Engine/Source/Editor/AddContentDialog/Private/ViewModels/AddContentWidgetViewModel.cpp b/Engine/Source/Editor/AddContentDialog/Private/ViewModels/AddContentWidgetViewModel.cpp
--- a/Engine/Source/Editor/AddContentDialog/Private/ViewModels/AddContentWidgetViewModel.cpp
+++ b/Engine/Source/Editor/AddContentDialog/Private/ViewModels/AddContentWidgetViewModel.cpp
@@ -114,6 +114,14 @@ void FAddContentWidgetViewModel::SetOnSelectedContentSourceChanged(FOnSelectedCo
 
 void FAddContentWidgetViewModel::BuildContentSourceViewModels()
 {
+	// Prevent reentrancy if the model sends a change event while we're getting the content sources
+	if (bBuildingContentSourceViewModels)
+	{
+		return;
+	}
+
+	TGuardValue GuardValue(bBuildingContentSourceViewModels, true);
+
 	Categories.Empty();
 	ContentSourceViewModels.Empty();
 	FilteredContentSourceViewModels.Empty();
diff --git a/Engine/Source/Editor/AddContentDialog/Private/ViewModels/AddContentWidgetViewModel.h b/Engine/Source/Editor/AddContentDialog/Private/ViewModels/AddContentWidgetViewModel.h
--- a/Engine/Source/Editor/AddContentDialog/Private/ViewModels/AddContentWidgetViewModel.h
+++ b/Engine/Source/Editor/AddContentDialog/Private/ViewModels/AddContentWidgetViewModel.h
@@ -103,6 +103,9 @@ private:
 	/** The current search text. */
 	FText SearchText;
 
+	/** True while building context source view models. */
+	bool bBuildingContentSourceViewModels = false;
+
 	/** The delegate which is executed when the available categories change. */
 	FOnCategoriesChanged OnCategoriesChanged;
 
